Status Bar updates
- Content browser drawer now working well with ctrl+space as the default keybinding. Summoning the drawer will immediately focus the search bar.  Had to add some new hooks into the content browser api for that
- Moved status bar to a proper module

diff --git a/Engine/Source/Editor/ContentBrowser/Private/ContentBrowserSingleton.cpp b/Engine/Source/Editor/ContentBrowser/Private/ContentBrowserSingleton.cpp
--- a/Engine/Source/Editor/ContentBrowser/Private/ContentBrowserSingleton.cpp
+++ b/Engine/Source/Editor/ContentBrowser/Private/ContentBrowserSingleton.cpp
@@ -226,6 +226,17 @@ void FContentBrowserSingleton::FocusPrimaryContentBrowser(bool bFocusSearch)
 	}
 }
 
+void FContentBrowserSingleton::FocusContentBrowserSearchField(TSharedRef<SWidget> ContentBrowserWidget)
+{
+	static FName ContentBrowserType(TEXT("SContentBrowser"));
+
+	if (ContentBrowserWidget->GetType() == ContentBrowserType)
+	{
+		TSharedRef<SContentBrowser> ContentBrowser = StaticCastSharedRef<SContentBrowser>(ContentBrowserWidget);
+		ContentBrowser->SetKeyboardFocusOnSearch();
+	}
+}
+
 void FContentBrowserSingleton::CreateNewAsset(const FString& DefaultAssetName, const FString& PackagePath, UClass* AssetClass, UFactory* Factory)
 {
 	FocusPrimaryContentBrowser(false);
diff --git a/Engine/Source/Editor/ContentBrowser/Private/ContentBrowserSingleton.h b/Engine/Source/Editor/ContentBrowser/Private/ContentBrowserSingleton.h
--- a/Engine/Source/Editor/ContentBrowser/Private/ContentBrowserSingleton.h
+++ b/Engine/Source/Editor/ContentBrowser/Private/ContentBrowserSingleton.h
@@ -57,6 +57,7 @@ public:
 	virtual FString CreateModalSaveAssetDialog(const FSaveAssetDialogConfig& SaveAssetConfig) override;
 	virtual bool HasPrimaryContentBrowser() const override;
 	virtual void FocusPrimaryContentBrowser(bool bFocusSearch) override;
+	virtual void FocusContentBrowserSearchField(TSharedRef<SWidget> ContentBrowserWidget) override;
 	virtual void CreateNewAsset(const FString& DefaultAssetName, const FString& PackagePath, UClass* AssetClass, UFactory* Factory) override;
 	virtual void SyncBrowserToAssets(const TArray<struct FAssetData>& AssetDataList, bool bAllowLockedBrowsers = false, bool bFocusContentBrowser = true, const FName& InstanceName = FName(), bool bNewSpawnBrowser = false) override;
 	virtual void SyncBrowserToAssets(const TArray<UObject*>& AssetList, bool bAllowLockedBrowsers = false, bool bFocusContentBrowser = true, const FName& InstanceName = FName(), bool bNewSpawnBrowser = false) override;
diff --git a/Engine/Source/Editor/ContentBrowser/Private/SContentBrowser.cpp b/Engine/Source/Editor/ContentBrowser/Private/SContentBrowser.cpp
--- a/Engine/Source/Editor/ContentBrowser/Private/SContentBrowser.cpp
+++ b/Engine/Source/Editor/ContentBrowser/Private/SContentBrowser.cpp
@@ -736,6 +736,7 @@ void SContentBrowser::Construct( const FArguments& InArgs, const FName& InInstan
 							.HintText( this, &SContentBrowser::GetSearchAssetsHintText )
 							.OnTextChanged( this, &SContentBrowser::OnSearchBoxChanged )
 							.OnTextCommitted( this, &SContentBrowser::OnSearchBoxCommitted )
+							.OnKeyDownHandler(this, &SContentBrowser::OnSearchKeyDown)
 							.OnAssetSearchBoxSuggestionFilter( this, &SContentBrowser::OnAssetSearchSuggestionFilter )
 							.OnAssetSearchBoxSuggestionChosen( this, &SContentBrowser::OnAssetSearchSuggestionChosen )
 							.DelayChangeNotificationsWhileTyping( true )
@@ -1744,6 +1745,19 @@ void SContentBrowser::OnSearchBoxCommitted(const FText& InSearchText, ETextCommi
 	SetSearchBoxText(InSearchText);
 }
 
+FReply SContentBrowser::OnSearchKeyDown(const FGeometry& Geometry, const FKeyEvent& InKeyEvent)
+{
+	FInputChord CheckChord(InKeyEvent.GetKey(), EModifierKey::FromBools(InKeyEvent.IsControlDown(), InKeyEvent.IsAltDown(), InKeyEvent.IsShiftDown(), InKeyEvent.IsCommandDown()));
+
+	// Clear focus if the content browser drawer key is clicked so it will close the opened content browser
+	if (FGlobalEditorCommonCommands::Get().OpenContentBrowserDrawer->HasActiveChord(CheckChord))
+	{
+		return FReply::Handled().ClearUserFocus(EFocusCause::SetDirectly);
+	}
+
+	return FReply::Unhandled();
+}
+
 bool SContentBrowser::IsSaveSearchButtonEnabled() const
 {
 	return !TextFilter->GetRawFilterText().IsEmptyOrWhitespace();
diff --git a/Engine/Source/Editor/ContentBrowser/Private/SContentBrowser.h b/Engine/Source/Editor/ContentBrowser/Private/SContentBrowser.h
--- a/Engine/Source/Editor/ContentBrowser/Private/SContentBrowser.h
+++ b/Engine/Source/Editor/ContentBrowser/Private/SContentBrowser.h
@@ -212,6 +212,9 @@ private:
 	/** Called by the editable text control when the user commits a text change */
 	void OnSearchBoxCommitted(const FText& InSearchText, ETextCommit::Type CommitInfo);
 
+	/** Called by the editable text control to allow the content browser to handle specific key presses without it resulting in typing into the search box */
+	FReply OnSearchKeyDown(const FGeometry& Geometry, const FKeyEvent& InKeyEvent);
+
 	/** Should the "Save Search" button be enabled? */
 	bool IsSaveSearchButtonEnabled() const;
 
diff --git a/Engine/Source/Editor/ContentBrowser/Public/IContentBrowserSingleton.h b/Engine/Source/Editor/ContentBrowser/Public/IContentBrowserSingleton.h
--- a/Engine/Source/Editor/ContentBrowser/Public/IContentBrowserSingleton.h
+++ b/Engine/Source/Editor/ContentBrowser/Public/IContentBrowserSingleton.h
@@ -517,6 +517,9 @@ public:
 	/** Brings the primary content browser to the front or opens one if it does not exist. */
 	virtual void FocusPrimaryContentBrowser(bool bFocusSearch) = 0;
 
+	/** Focuses the search field of a content browser widget */
+	virtual void FocusContentBrowserSearchField(TSharedRef<SWidget> ContentBrowserWidget) = 0;
+
 	/** Sets up an inline-name for the creation of a new asset in the primary content browser using the specified path and the specified class and/or factory */
 	virtual void CreateNewAsset(const FString& DefaultAssetName, const FString& PackagePath, UClass* AssetClass, UFactory* Factory) = 0;
 
diff --git a/Engine/Source/Editor/EditorStyle/Private/StarshipStyle.cpp b/Engine/Source/Editor/EditorStyle/Private/StarshipStyle.cpp
--- a/Engine/Source/Editor/EditorStyle/Private/StarshipStyle.cpp
+++ b/Engine/Source/Editor/EditorStyle/Private/StarshipStyle.cpp
@@ -7652,6 +7652,8 @@ void FStarshipEditorStyle::FStyle::SetupStatusBarStyle()
 	Set("StatusBar.ContentBrowserIcon", new CORE_IMAGE_BRUSH_SVG("Starship/Common/tile", Icon16x16));
 	Set("StatusBar.ContentBrowserUp", new CORE_IMAGE_BRUSH_SVG("Starship/Common/chevron-up", Icon16x16));
 	Set("StatusBar.ContentBrowserDown", new CORE_IMAGE_BRUSH_SVG("Starship/Common/chevron-down", Icon16x16));
+	Set("StatusBar.ContentBrowserShadow", new BOX_BRUSH("Starship/StatusBar/drawer-shadow", FMargin(10/64.0f, 20/64.f, 10/64.f, 0), FLinearColor(0,0,0,1)));
+
 	Set("StatusBar.Height", 32.0f);
 	Set("StatusBar.Background", new FSlateColorBrush(FStyleColors::Background));
 	Set("StatusBar.HelpIcon", new CORE_IMAGE_BRUSH_SVG("Starship/Common/help", Icon16x16, FStyleColors::Foreground));
diff --git a/Engine/Source/Editor/LevelEditor/LevelEditor.Build.cs b/Engine/Source/Editor/LevelEditor/LevelEditor.Build.cs
--- a/Engine/Source/Editor/LevelEditor/LevelEditor.Build.cs
+++ b/Engine/Source/Editor/LevelEditor/LevelEditor.Build.cs
@@ -68,6 +68,7 @@ public class LevelEditor : ModuleRules
 				"PixelInspectorModule",
 				"CommonMenuExtensions",
 				"ToolMenus",
+				"StatusBar",
 			}
 		);
 
diff --git a/Engine/Source/Editor/LevelEditor/Private/LevelEditorActions.cpp b/Engine/Source/Editor/LevelEditor/Private/LevelEditorActions.cpp
--- a/Engine/Source/Editor/LevelEditor/Private/LevelEditorActions.cpp
+++ b/Engine/Source/Editor/LevelEditor/Private/LevelEditorActions.cpp
@@ -3095,7 +3095,7 @@ void FLevelEditorCommands::RegisterCommands()
 	UI_COMMAND( SnapOriginToGridPerActor, "Snap Origin to Grid Per Actor", "Snaps each selected actor separately to the nearest grid location at its origin", EUserInterfaceActionType::Button, FInputChord() );
 	UI_COMMAND( AlignOriginToGrid, "Align Origin to Grid", "Aligns the actor to the nearest grid location at its origin", EUserInterfaceActionType::Button, FInputChord() );
 
-	UI_COMMAND( SnapTo2DLayer, "Snap to 2D Layer", "Snaps the actor to the current 2D snap layer", EUserInterfaceActionType::Button, FInputChord( EModifierKey::Control, EKeys::SpaceBar ) );
+	UI_COMMAND( SnapTo2DLayer, "Snap to 2D Layer", "Snaps the actor to the current 2D snap layer", EUserInterfaceActionType::Button, FInputChord() );
 	UI_COMMAND( MoveSelectionUpIn2DLayers, "Bring selection forward a snap layer", "Bring selection forward a snap layer", EUserInterfaceActionType::Button, FInputChord(EKeys::PageUp, EModifierKey::Control) );
 	UI_COMMAND( MoveSelectionDownIn2DLayers, "Send selection backward a snap layer", "Send selection backward a snap layer", EUserInterfaceActionType::Button, FInputChord(EKeys::PageDown, EModifierKey::Control) );
 	UI_COMMAND( MoveSelectionToTop2DLayer, "Bring selection to the front snap layer", "Bring selection to the front snap layer", EUserInterfaceActionType::Button, FInputChord(EKeys::PageUp, EModifierKey::Shift | EModifierKey::Control) );
diff --git a/Engine/Source/Editor/LevelEditor/Private/SLevelEditor.cpp b/Engine/Source/Editor/LevelEditor/Private/SLevelEditor.cpp
--- a/Engine/Source/Editor/LevelEditor/Private/SLevelEditor.cpp
+++ b/Engine/Source/Editor/LevelEditor/Private/SLevelEditor.cpp
@@ -52,7 +52,7 @@
 #include "EditorViewportCommands.h"
 #include "IPlacementModeModule.h"
 #include "Classes/EditorStyleSettings.h"
-#include "Subsystems/StatusBarSubsystem.h"
+#include "StatusBarSubsystem.h"
 
 static const FName MainFrameModuleName("MainFrame");
 static const FName LevelEditorModuleName("LevelEditor");
@@ -197,6 +197,8 @@ void SLevelEditor::Initialize( const TSharedRef<SDockTab>& OwnerTab, const TShar
 
 	FLevelEditorModule& LevelEditorModule = FModuleManager::GetModuleChecked< FLevelEditorModule >(LevelEditorModuleName);
 
+	FModuleManager::Get().LoadModuleChecked("StatusBar");
+
 	LevelEditorModule.OnActorSelectionChanged().AddSP(this, &SLevelEditor::OnActorSelectionChanged);
 
 	TSharedRef<SWidget> ContentArea = RestoreContentArea( OwnerTab, OwnerWindow );
diff --git a/Engine/Source/Editor/LevelEditor/Private/SLevelEditorToolBox.cpp b/Engine/Source/Editor/LevelEditor/Private/SLevelEditorToolBox.cpp
--- a/Engine/Source/Editor/LevelEditor/Private/SLevelEditorToolBox.cpp
+++ b/Engine/Source/Editor/LevelEditor/Private/SLevelEditorToolBox.cpp
@@ -17,7 +17,7 @@
 #include "EdMode.h"
 #include "Widgets/Docking/SDockTab.h"
 #include "Widgets/Layout/SExpandableArea.h"
-#include "Subsystems/StatusBarSubsystem.h"
+#include "StatusBarSubsystem.h"
 
 #define LOCTEXT_NAMESPACE "SLevelEditorToolBox"
 
diff --git a/Engine/Source/Editor/LevelEditor/Private/SLevelEditorToolBox.h b/Engine/Source/Editor/LevelEditor/Private/SLevelEditorToolBox.h
--- a/Engine/Source/Editor/LevelEditor/Private/SLevelEditorToolBox.h
+++ b/Engine/Source/Editor/LevelEditor/Private/SLevelEditorToolBox.h
@@ -11,7 +11,7 @@
 #include "Editor.h"
 #include "ILevelEditor.h"
 #include "Misc/NotifyHook.h"
-#include "Subsystems/StatusBarSubsystem.h"
+#include "StatusBarSubsystem.h"
 
 class FExtender;
 class SBorder;
diff --git a/Engine/Source/Editor/StatusBar/Private/SStatusBar.cpp b/Engine/Source/Editor/StatusBar/Private/SStatusBar.cpp
new file mode 100644
--- /dev/null
+++ b/Engine/Source/Editor/StatusBar/Private/SStatusBar.cpp
@@ -0,0 +1,555 @@
+// Copyright Epic Games, Inc. All Rights Reserved.
+
+#include "SStatusBar.h"
+#include "Widgets/Images/SImage.h"
+#include "Widgets/Input/SButton.h"
+#include "Widgets/Layout/SBorder.h"
+#include "Widgets/Layout/SBox.h"
+#include "Widgets/Input/SComboButton.h"
+#include "Widgets/Images/SImage.h"
+#include "Widgets/Text/SRichTextBlock.h"
+#include "Widgets/Layout/SBox.h"
+#include "Widgets/Layout/SBorder.h"
+#include "Widgets/Input/SButton.h"
+#include "Widgets/Input/SMultiLineEditableTextBox.h"
+#include "Widgets/Layout/SSpacer.h"
+#include "Widgets/Docking/SDockTab.h"
+#include "TimerManager.h"
+#include "Framework/Commands/Commands.h"
+#include "ToolMenuContext.h"
+#include "ToolMenus.h"
+#include "OutputLogModule.h"
+
+#define LOCTEXT_NAMESPACE "StatusBar"
+
+class SContentBrowserOverlay : public SCompoundWidget
+{
+	SLATE_BEGIN_ARGS(SContentBrowserOverlay)
+	{
+		_Clipping = EWidgetClipping::ClipToBounds;
+		_ShadowOffset = FVector2D(10.0f, 20.0f);
+	}
+	SLATE_DEFAULT_SLOT(FArguments, Content)
+		SLATE_ARGUMENT(float, MaxContentBrowserHeight)
+		SLATE_ARGUMENT(FVector2D, ShadowOffset)
+		SLATE_END_ARGS()
+
+
+		void SetHeight(float InHeight)
+	{
+		Height = InHeight;
+	}
+
+	void Construct(const FArguments& InArgs)
+	{
+		Height = 0;
+
+		ShadowOffset = InArgs._ShadowOffset;
+		MaxContentBrowserHeight = InArgs._MaxContentBrowserHeight;
+
+		BackgroundBrush = FAppStyle::Get().GetBrush("StatusBar.Background");
+		ShadowBrush = FAppStyle::Get().GetBrush("StatusBar.ContentBrowserShadow");
+
+		ChildSlot
+			[
+				InArgs._Content.Widget
+			];
+	}
+
+	virtual bool SupportsKeyboardFocus() const
+	{
+		return true;
+	}
+
+	virtual FVector2D ComputeDesiredSize(float) const
+	{
+		return FVector2D(1.0f, MaxContentBrowserHeight + ShadowOffset.Y);
+	}
+
+	virtual void OnArrangeChildren(const FGeometry& AllottedGeometry, FArrangedChildren& ArrangedChildren) const override
+	{
+		const EVisibility ChildVisibility = ChildSlot.GetWidget()->GetVisibility();
+		if (ArrangedChildren.Accepts(ChildVisibility))
+		{
+			ArrangedChildren.AddWidget(
+				AllottedGeometry.MakeChild(
+					ChildSlot.GetWidget(),
+					ShadowOffset,
+					FVector2D(AllottedGeometry.GetLocalSize().X - (ShadowOffset.X * 2), MaxContentBrowserHeight)
+				)
+			);
+		}
+	}
+
+	virtual int32 OnPaint(const FPaintArgs& Args, const FGeometry& AllottedGeometry, const FSlateRect& MyCullingRect, FSlateWindowElementList& OutDrawElements, int32 LayerId, const FWidgetStyle& InWidgetStyle, bool bParentEnabled) const override
+	{
+		static FLinearColor ShadowColor = FAppStyle::Get().GetColor("Colors.Foldout");
+
+		FGeometry RenderTransformedChildGeometry = AllottedGeometry.MakeChild(FSlateRenderTransform(FVector2D(0.0f, MaxContentBrowserHeight - Height)));
+
+		FSlateDrawElement::MakeBox(
+			OutDrawElements,
+			LayerId,
+			RenderTransformedChildGeometry.ToPaintGeometry(),
+			ShadowBrush,
+			ESlateDrawEffect::None,
+			ShadowBrush->GetTint(InWidgetStyle));
+
+		FSlateDrawElement::MakeBox(
+			OutDrawElements,
+			LayerId,
+			RenderTransformedChildGeometry.ToPaintGeometry(ShadowOffset, FVector2D(AllottedGeometry.GetLocalSize().X - (ShadowOffset.X * 2), MaxContentBrowserHeight)),
+			BackgroundBrush,
+			ESlateDrawEffect::None,
+			BackgroundBrush->GetTint(InWidgetStyle));
+
+		int32 OutLayerId = SCompoundWidget::OnPaint(Args, RenderTransformedChildGeometry, MyCullingRect, OutDrawElements, LayerId, InWidgetStyle, bParentEnabled);
+
+		return OutLayerId;
+
+	}
+private:
+	const FSlateBrush* BackgroundBrush;
+	const FSlateBrush* ShadowBrush;
+	FVector2D ShadowOffset;
+	float Height;
+	float MaxContentBrowserHeight;
+
+};
+
+
+void SStatusBar::Construct(const FArguments& InArgs, FName InStatusBarName, const TSharedRef<SDockTab> InParentTab)
+{
+	StatusBarName = InStatusBarName;
+	ParentTab = InParentTab;
+
+	UpArrow = FAppStyle::Get().GetBrush("StatusBar.ContentBrowserUp");
+	DownArrow = FAppStyle::Get().GetBrush("StatusBar.ContentBrowserDown");
+
+	const FSlateBrush* StatusBarBackground = FAppStyle::Get().GetBrush("StatusBar.Background");
+
+	GetContentBrowserDelegate = InArgs._OnGetContentBrowser;
+	OnContentBrowserOpenedDelegate = InArgs._OnContentBrowserOpened;
+	OnContentBrowserDismissedDelegate = InArgs._OnContentBrowserDismissed;
+
+	ContentBrowserEasingCurve = FCurveSequence(0.0f, 0.15f, ECurveEaseFunction::QuadOut);
+
+	FSlateApplication::Get().OnFocusChanging().AddSP(this, &SStatusBar::OnGlobalFocusChanging);
+	FGlobalTabmanager::Get()->OnActiveTabChanged_Subscribe(FOnActiveTabChanged::FDelegate::CreateSP(this, &SStatusBar::OnActiveTabChanged));
+	FGlobalTabmanager::Get()->OnTabForegrounded_Subscribe(FOnActiveTabChanged::FDelegate::CreateSP(this, &SStatusBar::OnActiveTabChanged));
+
+	ChildSlot
+	[
+		SNew(SBox)
+		.HeightOverride(FAppStyle::Get().GetFloat("StatusBar.Height"))
+		[				SNew(SHorizontalBox)
+
+			+ SHorizontalBox::Slot()
+			.Padding(1.0f, 0.0f)
+			.AutoWidth()
+			[
+				SNew(SBorder)
+				.Padding(FMargin(2.0f, 0.0f))
+				.BorderImage(StatusBarBackground)
+				.VAlign(VAlign_Center)
+				[			
+					MakeContentBrowserWidget()
+				]
+			] 
+			+ SHorizontalBox::Slot()
+			.AutoWidth()
+			.Padding(1.0f, 0.0f)
+			[
+				SNew(SBorder)
+				.Padding(0.0f)
+				.BorderImage(StatusBarBackground)
+				.VAlign(VAlign_Center)
+				.Padding(FMargin(6.0f, 0.0f))
+				[
+					MakeDebugConsoleWidget(InArgs._OnConsoleClosed)
+				]
+			]
+			+ SHorizontalBox::Slot()
+			.FillWidth(1.0f)
+			.Padding(1.0f, 0.0f)
+			[
+				SNew(SBorder)
+				.Padding(0.0f)
+				.BorderImage(StatusBarBackground)
+				.VAlign(VAlign_Center)
+				.Padding(FMargin(6.0f, 0.0f))
+				[
+					MakeStatusMessageWidget()
+				]
+			]
+			+ SHorizontalBox::Slot()
+			.HAlign(HAlign_Right)
+			.AutoWidth()
+			.Padding(1.0f, 0.0f)
+			[
+				SNew(SBorder)
+				.Padding(0.0f)
+				.BorderImage(StatusBarBackground)
+				.VAlign(VAlign_Center)
+				[
+					MakeStatusBarToolBarWidget()
+				]
+			]
+		]
+	];
+}
+
+void SStatusBar::PushMessage(FStatusBarMessageHandle& InHandle, const TAttribute<FText>& InMessage, const TAttribute<FText>& InHintText)
+{
+	MessageStack.Emplace(InMessage, InHintText, InHandle);
+}
+
+void SStatusBar::PopMessage(FStatusBarMessageHandle& InHandle)
+{
+	if (InHandle.IsValid() && MessageStack.Num() > 0)
+	{
+		MessageStack.RemoveAll([InHandle](const FStatusBarMessage& Message)
+			{
+				return Message.Handle == InHandle;
+			});
+	}
+}
+
+void SStatusBar::ClearAllMessages()
+{
+	MessageStack.Empty();
+}
+
+EVisibility SStatusBar::GetHelpIconVisibility() const
+{
+	if (MessageStack.Num() > 0)
+	{
+		const FStatusBarMessage& MessageData = MessageStack.Top();
+
+		const FText& Message = MessageData.MessageText.Get();
+		const FText& HintText = MessageData.HintText.Get();
+
+		return (!Message.IsEmpty() || !HintText.IsEmpty()) ? EVisibility::SelfHitTestInvisible : EVisibility::Collapsed;
+	}
+
+	return EVisibility::Collapsed;
+
+}
+
+TSharedPtr<SDockTab> SStatusBar::GetParentTab() const
+{
+	return ParentTab.Pin();
+}
+
+void SStatusBar::FocusDebugConsole()
+{
+	FSlateApplication::Get().SetKeyboardFocus(ConsoleEditBox, EFocusCause::SetDirectly);
+}
+
+bool SStatusBar::IsDebugConsoleFocused() const
+{
+	return ConsoleEditBox->HasKeyboardFocus();
+}
+
+void SStatusBar::OnGlobalFocusChanging(const FFocusEvent& FocusEvent, const FWeakWidgetPath& OldFocusedWidgetPath, const TSharedPtr<SWidget>& OldFocusedWidget, const FWidgetPath& NewFocusedWidgetPath, const TSharedPtr<SWidget>& NewFocusedWidget)
+{
+	// Sometimes when dismissing focus can change which will trigger this again
+	static bool bIsRentrant = false;
+
+	if(!bIsRentrant)
+	{
+		TGuardValue<bool> RentrancyGuard(bIsRentrant, true);
+
+		TSharedRef<SWidget> ThisWidget = AsShared();
+
+		
+		bool bShouldDismiss = false;
+		// Do not close due to slow tasks as those opening send window activation events
+		if (!GIsSlowTask && !FSlateApplication::Get().GetActiveModalWindow().IsValid() && ContentBrowserOverlayContent.IsValid() && (!NewFocusedWidgetPath.ContainsWidget(ContentBrowserOverlayContent.ToSharedRef()) && !NewFocusedWidgetPath.ContainsWidget(ThisWidget)))
+		{
+			if (TSharedPtr<SWidget> MenuHost = FSlateApplication::Get().GetMenuHostWidget())
+			{
+				FWidgetPath MenuHostPath;
+
+				// See if the menu being opened is part of the content browser path and if so the menu should not be dismissed
+				FSlateApplication::Get().GeneratePathToWidgetUnchecked(MenuHost.ToSharedRef(), MenuHostPath);
+				if (!MenuHostPath.ContainsWidget(ContentBrowserOverlayContent.ToSharedRef()))
+				{
+					bShouldDismiss = true;
+				}
+			}
+			else
+			{
+				bShouldDismiss = true;
+			}
+		}
+
+		if (bShouldDismiss)
+		{
+			DismissContentBrowser(NewFocusedWidget);
+		}
+	}
+}
+
+void SStatusBar::OnActiveTabChanged(TSharedPtr<SDockTab> PreviouslyActive, TSharedPtr<SDockTab> NewlyActivated)
+{
+	// Remove the content browser as the tab well's content is changing 
+	RemoveContentBrowser();
+}
+
+bool SStatusBar::IsContentBrowserOpened() const
+{
+	return ContentBrowserOverlayContent.IsValid();
+}
+
+void SStatusBar::OpenContentBrowser()
+{
+	if (!IsContentBrowserOpened())
+	{
+		TSharedRef<SStatusBar> ThisStatusBar = SharedThis(this);
+
+		TSharedPtr<SWindow> MyWindow = FSlateApplication::Get().FindWidgetWindow(AsShared());
+
+		MaxContentBrowserHeight = MyWindow->GetSizeInScreen().Y * 0.33f+20;
+		MinContentBrowserHeight = GetTickSpaceGeometry().GetLocalSize().Y + MyWindow->GetWindowBorderSize().Bottom;
+		ContentBrowserEasingCurve.Play(ThisStatusBar, false, ContentBrowserEasingCurve.IsPlaying() ? ContentBrowserEasingCurve.GetSequenceTime() : 0.0f, false);
+
+		MyWindow->AddOverlaySlot()
+			.VAlign(VAlign_Bottom)
+			.Padding(FMargin(10.0f,20.0f,10.0f, MinContentBrowserHeight))
+			[
+				SAssignNew(ContentBrowserOverlayContent, SContentBrowserOverlay)
+				.MaxContentBrowserHeight(MaxContentBrowserHeight)
+				[
+					GetContentBrowserDelegate.Execute()
+				]
+			];
+
+		if (!ContentBrowserOpenCloseTimer.IsValid())
+		{
+			ContentBrowserOpenCloseTimer = RegisterActiveTimer(0.0f, FWidgetActiveTimerDelegate::CreateSP(ThisStatusBar, &SStatusBar::UpdateContentBrowserAnimation));
+		}
+
+		OnContentBrowserOpenedDelegate.ExecuteIfBound(ThisStatusBar);
+	}
+}
+
+/**
+ * Dismisses an open content browser with an animation.  The content browser is removed once the animation is complete
+ */
+void SStatusBar::DismissContentBrowser(const TSharedPtr<SWidget>& NewlyFocusedWidget)
+{
+	if (ContentBrowserOverlayContent.IsValid())
+	{
+		if (ContentBrowserEasingCurve.IsForward())
+		{
+			ContentBrowserEasingCurve.Reverse();
+		}
+
+		if (!ContentBrowserOpenCloseTimer.IsValid())
+		{
+			ContentBrowserOpenCloseTimer = RegisterActiveTimer(0.0f, FWidgetActiveTimerDelegate::CreateSP(this, &SStatusBar::UpdateContentBrowserAnimation));
+		}
+
+		OnContentBrowserDismissedDelegate.ExecuteIfBound(NewlyFocusedWidget);
+	}
+}
+
+/**
+ * Removes a content browser from the layout immediately
+ */
+void SStatusBar::RemoveContentBrowser()
+{
+	if(ContentBrowserOverlayContent.IsValid())
+	{
+		TSharedPtr<SWindow> Window = FSlateApplication::Get().FindWidgetWindow(SharedThis(this));
+
+		Window->RemoveOverlaySlot(ContentBrowserOverlayContent.ToSharedRef());
+
+		ContentBrowserOverlayContent.Reset();
+
+		if (ContentBrowserOpenCloseTimer.IsValid())
+		{
+			UnRegisterActiveTimer(ContentBrowserOpenCloseTimer.ToSharedRef());
+		}
+
+		ContentBrowserOpenCloseTimer.Reset();
+		ContentBrowserEasingCurve.JumpToStart();
+		ContentBrowserEasingCurve.Pause();
+	}
+}
+
+const FSlateBrush* SStatusBar::GetContentBrowserExpandArrowImage() const
+{
+	return ContentBrowserOverlayContent.IsValid() ? DownArrow : UpArrow;
+}
+
+FText SStatusBar::GetStatusBarMessage() const
+{
+	FText FullMessage;
+	if (MessageStack.Num() > 0)
+	{
+		const FStatusBarMessage& MessageData = MessageStack.Top();
+
+		const FText& Message = MessageData.MessageText.Get();
+		const FText& HintText = MessageData.HintText.Get();
+
+		FullMessage = HintText.IsEmpty() ? Message : FText::Format(LOCTEXT("StatusBarMessageFormat", "{0} <StatusBar.Message.InHintText>{1}</>"), Message, HintText);
+	}
+
+	return FullMessage;
+}
+
+TSharedRef<SWidget> SStatusBar::MakeContentBrowserWidget()
+{
+	return
+		SNew(SButton)
+		.ButtonStyle(&FAppStyle::Get().GetWidgetStyle<FButtonStyle>("StatusBar.StatusBarButton"))
+		.OnClicked(this, &SStatusBar::OnContentBrowserButtonClicked)
+		[
+			SNew(SHorizontalBox)
+			+ SHorizontalBox::Slot()
+			.HAlign(HAlign_Left)
+			.VAlign(VAlign_Center)
+			.Padding(2.0f)
+			.AutoWidth()
+			[
+				SNew(SImage)
+				.ColorAndOpacity(FSlateColor::UseForeground())
+				.Image(this, &SStatusBar::GetContentBrowserExpandArrowImage)
+			]
+			+ SHorizontalBox::Slot()
+			.Padding(2.0f)
+			.HAlign(HAlign_Left)
+			.VAlign(VAlign_Center)
+			.AutoWidth()
+			[
+				SNew(SImage)
+				.ColorAndOpacity(FSlateColor::UseForeground())
+				.Image(FAppStyle::Get().GetBrush("StatusBar.ContentBrowserIcon"))
+			]
+			+ SHorizontalBox::Slot()
+			.VAlign(VAlign_Center)
+			.Padding(2.0f)
+			[
+				SNew(STextBlock)
+				.TextStyle(&FAppStyle::Get().GetWidgetStyle<FTextBlockStyle>("NormalText"))
+				.Text(LOCTEXT("StatusBar_ContentBrowserButton", "Content Browser"))
+			]
+		];	
+}
+
+TSharedRef<SWidget> SStatusBar::MakeStatusBarToolBarWidget()
+{
+	RegisterStatusBarMenu();
+
+	FToolMenuContext MenuContext;
+	RegisterSourceControlStatus();
+
+	return UToolMenus::Get()->GenerateWidget("StatusBar.ToolBar", MenuContext);
+}
+
+TSharedRef<SWidget> SStatusBar::MakeDebugConsoleWidget(FSimpleDelegate OnConsoleClosed)
+{
+	FOutputLogModule& OutputLogModule = FModuleManager::LoadModuleChecked<FOutputLogModule>(TEXT("OutputLog"));
+
+	return
+		SNew(SBox)
+		.WidthOverride(350.f)
+		[
+			OutputLogModule.MakeConsoleInputBox(ConsoleEditBox, OnConsoleClosed)
+		];
+}
+
+TSharedRef<SWidget> SStatusBar::MakeStatusMessageWidget()
+{
+	return 
+		SNew(SHorizontalBox)
+		+ SHorizontalBox::Slot()
+		.AutoWidth()
+		.VAlign(VAlign_Center)
+		.HAlign(HAlign_Center)
+		[
+			SNew(SImage)
+			.ColorAndOpacity(FSlateColor::UseForeground())
+			.Image(FAppStyle::Get().GetBrush("StatusBar.HelpIcon"))
+			.Visibility(this, &SStatusBar::GetHelpIconVisibility)
+		]
+		+ SHorizontalBox::Slot()
+		.VAlign(VAlign_Center)
+		.Padding(5.0f, 0.0f)
+		[
+			SNew(SRichTextBlock)
+			.TextStyle(&FAppStyle::Get().GetWidgetStyle<FTextBlockStyle>("StatusBar.Message.MessageText"))
+			.Text(this, &SStatusBar::GetStatusBarMessage)
+			.DecoratorStyleSet(&FAppStyle::Get())
+		];
+}
+
+FReply SStatusBar::OnContentBrowserButtonClicked()
+{
+	if(!ContentBrowserOverlayContent.IsValid())
+	{
+		OpenContentBrowser();
+	}
+	else
+	{
+		DismissContentBrowser(nullptr);
+	}
+
+	return FReply::Handled();
+}
+
+EActiveTimerReturnType SStatusBar::UpdateContentBrowserAnimation(double CurrentTime, float DeltaTime)
+{
+	check(ContentBrowserOverlayContent.IsValid());
+
+	ContentBrowserOverlayContent->SetHeight(FMath::Lerp(0.0f, MaxContentBrowserHeight, ContentBrowserEasingCurve.GetLerp()));
+
+	if (!ContentBrowserEasingCurve.IsPlaying())
+	{
+		if(ContentBrowserEasingCurve.IsAtStart())
+		{
+			RemoveContentBrowser();
+		}
+
+		ContentBrowserOpenCloseTimer.Reset();
+		return EActiveTimerReturnType::Stop;
+	}
+
+	return EActiveTimerReturnType::Continue;
+}
+
+void SStatusBar::RegisterStatusBarMenu()
+{
+	static const FName StatusBarToolBarName("StatusBar.ToolBar");
+	UToolMenus* ToolMenus = UToolMenus::Get();
+	if (ToolMenus->IsMenuRegistered(StatusBarToolBarName))
+	{
+		return;
+	}
+
+	UToolMenu* ToolBar = ToolMenus->RegisterMenu(StatusBarToolBarName, NAME_None, EMultiBoxType::SlimHorizontalToolBar);
+	ToolBar->StyleName = "StatusBarToolBar";
+}
+
+void SStatusBar::RegisterSourceControlStatus()
+{
+	// Source Control preferences
+	FSourceControlMenuHelpers::CheckSourceControlStatus();
+	{
+		UToolMenu* SourceControlMenu = UToolMenus::Get()->ExtendMenu("StatusBar.ToolBar");
+		FToolMenuSection& Section = SourceControlMenu->FindOrAddSection("SourceControl");
+
+		Section.AddEntry(
+			FToolMenuEntry::InitWidget(
+				"SourceControl",
+				FSourceControlMenuHelpers::MakeSourceControlStatusWidget(),
+				FText::GetEmpty(),
+				true,
+				false
+			));
+	}
+}
+
+#undef LOCTEXT_NAMESPACE
\ No newline at end of file
diff --git a/Engine/Source/Editor/StatusBar/Private/SourceControlMenuHelpers.cpp b/Engine/Source/Editor/StatusBar/Private/SourceControlMenuHelpers.cpp
new file mode 100644
--- /dev/null
+++ b/Engine/Source/Editor/StatusBar/Private/SourceControlMenuHelpers.cpp
@@ -0,0 +1,282 @@
+// Copyright Epic Games, Inc. All Rights Reserved.
+
+#include "ISourceControlOperation.h"
+#include "SourceControlOperations.h"
+#include "ISourceControlProvider.h"
+#include "ISourceControlModule.h"
+#include "SourceControlMenuHelpers.h"
+#include "SourceControlWindows.h"
+#include "FileHelpers.h"
+#include "ToolMenuContext.h"
+#include "ToolMenus.h"
+
+#define LOCTEXT_NAMESPACE "SourceControlCommands"
+
+TSharedRef<FUICommandList> FSourceControlCommands::ActionList(new FUICommandList());
+
+FSourceControlCommands::FSourceControlCommands() 
+	: TCommands<FSourceControlCommands>
+(
+	"SourceControl",
+	NSLOCTEXT("Contexts", "SourceControl", "Source Control"),
+	"LevelEditor",
+	FEditorStyle::GetStyleSetName()
+)
+{}
+
+/**
+ * Initialize commands
+ */
+void FSourceControlCommands::RegisterCommands()
+{
+	UI_COMMAND(ConnectToSourceControl, "Connect to Source Control...", "Opens a dialog to connect to source control.", EUserInterfaceActionType::Button, FInputChord());
+	UI_COMMAND(ChangeSourceControlSettings, "Change Source Control Settings...", "Opens a dialog to change source control settings.", EUserInterfaceActionType::Button, FInputChord());
+	UI_COMMAND(CheckOutModifiedFiles, "Check Out Modified Files...", "Opens a dialog to check out any assets which have been modified.", EUserInterfaceActionType::Button, FInputChord());
+	UI_COMMAND(SubmitToSourceControl, "Submit to Source Control...", "Opens a dialog with check in options for content and levels.",   EUserInterfaceActionType::Button, FInputChord());
+
+	ActionList->MapAction(
+		ConnectToSourceControl,
+		FExecuteAction::CreateStatic(&FSourceControlCommands::ConnectToSourceControl_Clicked)                        
+	);
+
+	ActionList->MapAction(
+		ChangeSourceControlSettings,
+		FExecuteAction::CreateStatic(&FSourceControlCommands::ConnectToSourceControl_Clicked)
+	);
+
+	ActionList->MapAction(
+		CheckOutModifiedFiles,
+		FExecuteAction::CreateStatic(&FSourceControlCommands::CheckOutModifiedFiles_Clicked),
+		FCanExecuteAction::CreateStatic(&FSourceControlCommands::CheckOutModifiedFiles_CanExecute)
+	);
+
+	ActionList->MapAction(
+		SubmitToSourceControl,
+		FExecuteAction::CreateStatic(&FSourceControlCommands::SubmitToSourceControl_Clicked),
+		FCanExecuteAction::CreateStatic(&FSourceControlCommands::SubmitToSourceControl_CanExecute)
+	);
+
+}
+
+void FSourceControlCommands::ConnectToSourceControl_Clicked()
+{
+	// Show login window regardless of current status - its useful as a shortcut to change settings.
+	ISourceControlModule& SourceControlModule = ISourceControlModule::Get();
+	SourceControlModule.ShowLoginDialog(FSourceControlLoginClosed(), ELoginWindowMode::Modeless, EOnLoginWindowStartup::PreserveProvider);
+}
+
+bool FSourceControlCommands::CheckOutModifiedFiles_CanExecute()
+{
+	ISourceControlModule& SourceControlModule = ISourceControlModule::Get();
+	if (ISourceControlModule::Get().IsEnabled() &&
+		ISourceControlModule::Get().GetProvider().IsAvailable())
+	{
+		TArray<UPackage*> PackagesToSave;
+		FEditorFileUtils::GetDirtyWorldPackages(PackagesToSave);
+		FEditorFileUtils::GetDirtyContentPackages(PackagesToSave);
+
+		return PackagesToSave.Num() > 0;
+	}
+
+	return false;
+}
+
+void FSourceControlCommands::CheckOutModifiedFiles_Clicked()
+{
+	TArray<UPackage*> PackagesToSave;
+	FEditorFileUtils::GetDirtyWorldPackages(PackagesToSave);
+	FEditorFileUtils::GetDirtyContentPackages(PackagesToSave);
+
+	const bool bCheckDirty = true;
+	const bool bPromptUserToSave = false;
+	FEditorFileUtils::PromptForCheckoutAndSave(PackagesToSave, bCheckDirty, bPromptUserToSave);
+}
+
+bool FSourceControlCommands::SubmitToSourceControl_CanExecute()
+{
+	ISourceControlModule& SourceControlModule = ISourceControlModule::Get();
+	return ISourceControlModule::Get().IsEnabled() &&
+		ISourceControlModule::Get().GetProvider().IsAvailable() &&
+		FSourceControlWindows::CanChoosePackagesToCheckIn();
+}
+
+void FSourceControlCommands::SubmitToSourceControl_Clicked()
+{
+	FSourceControlWindows::ChoosePackagesToCheckIn();
+}
+
+
+FSourceControlMenuHelpers::EQueryState FSourceControlMenuHelpers::QueryState = FSourceControlMenuHelpers::EQueryState::NotQueried;
+
+void FSourceControlMenuHelpers::CheckSourceControlStatus()
+{
+	ISourceControlModule& SourceControlModule = ISourceControlModule::Get();
+	if (SourceControlModule.IsEnabled())
+	{
+		SourceControlModule.GetProvider().Execute(ISourceControlOperation::Create<FConnect>(),
+			EConcurrency::Asynchronous,
+			FSourceControlOperationComplete::CreateStatic(&FSourceControlMenuHelpers::OnSourceControlOperationComplete));
+		QueryState = EQueryState::Querying;
+	}
+}
+
+void FSourceControlMenuHelpers::OnSourceControlOperationComplete(const FSourceControlOperationRef& InOperation, ECommandResult::Type InResult)
+{
+	QueryState = EQueryState::Queried;
+}
+
+TSharedRef<SWidget> FSourceControlMenuHelpers::GenerateSourceControlMenuContent()
+{
+	UToolMenu* SourceControlMenu = UToolMenus::Get()->RegisterMenu("StatusBar.ToolBar.SourceControl", NAME_None, EMultiBoxType::Menu, false);
+
+	FToolMenuSection& Section = SourceControlMenu->AddSection("SourceControlActions", LOCTEXT("SourceControlMenuHeadingActions", "Actions"));
+
+	Section.AddDynamicEntry("ConnectToSourceControl", FNewToolMenuSectionDelegate::CreateLambda([](FToolMenuSection& InSection)
+		{
+			ISourceControlModule& SourceControlModule = ISourceControlModule::Get();
+			if (ISourceControlModule::Get().IsEnabled() && ISourceControlModule::Get().GetProvider().IsAvailable())
+			{
+				InSection.AddMenuEntry(
+					FSourceControlCommands::Get().ChangeSourceControlSettings,
+					TAttribute<FText>(),
+					TAttribute<FText>(),
+					FSlateIcon(FEditorStyle::GetStyleSetName(), "SourceControl.Actions.ChangeSettings")
+				);
+			}
+			else
+			{
+				InSection.AddMenuEntry(
+					FSourceControlCommands::Get().ConnectToSourceControl,
+					TAttribute<FText>(),
+					TAttribute<FText>(),
+					FSlateIcon(FEditorStyle::GetStyleSetName(), "SourceControl.Actions.Connect")
+				);
+			}
+		}));
+
+	Section.AddSeparator("SourceControlConnectionSeparator");
+
+	Section.AddMenuEntry(
+		FSourceControlCommands::Get().CheckOutModifiedFiles,
+		TAttribute<FText>(),
+		TAttribute<FText>(),
+		FSlateIcon(FEditorStyle::GetStyleSetName(), "SourceControl.Actions.CheckOut")
+	);
+
+
+	Section.AddMenuEntry(
+		FSourceControlCommands::Get().SubmitToSourceControl,
+		TAttribute<FText>(),
+		TAttribute<FText>(),
+		FSlateIcon(FEditorStyle::GetStyleSetName(), "SourceControl.Actions.Submit")
+	);
+
+	return UToolMenus::Get()->GenerateWidget("StatusBar.ToolBar.SourceControl", FToolMenuContext(FSourceControlCommands::ActionList));
+}
+
+FText FSourceControlMenuHelpers::GetSourceControlStatusText()
+{
+	if (QueryState == EQueryState::Querying)
+	{
+		return LOCTEXT("SourceControlStatus_Querying", "Contacting Server....");
+	}
+	else
+	{
+		ISourceControlModule& SourceControlModule = ISourceControlModule::Get();
+		if (SourceControlModule.IsEnabled())
+		{
+			if (!SourceControlModule.GetProvider().IsAvailable())
+			{
+				return LOCTEXT("SourceControlStatus_Error", "Server Unavailable");
+			}
+			else
+			{
+				return LOCTEXT("SourceControlStatus_Available", "Source Control");
+			}
+		}
+		else
+		{
+			return LOCTEXT("SourceControlStatus_Error", "Source Control Off");
+		}
+	}
+}
+FText FSourceControlMenuHelpers::GetSourceControlTooltip()
+{
+	if (QueryState == EQueryState::Querying)
+	{
+		return LOCTEXT("SourceControlUnknown", "Source control status is unknown");
+	}
+	else
+	{
+		return ISourceControlModule::Get().GetProvider().GetStatusText();
+	}
+}
+
+const FSlateBrush* FSourceControlMenuHelpers::GetSourceControlIcon()
+{
+
+	if (QueryState == EQueryState::Querying)
+	{
+		static const FSlateBrush* QueryBrush = FAppStyle::Get().GetBrush("SourceControl.StatusIcon.Unknown");
+		return QueryBrush;
+	}
+	else
+	{
+		ISourceControlModule& SourceControlModule = ISourceControlModule::Get();
+		if (SourceControlModule.IsEnabled())
+		{
+			if (!SourceControlModule.GetProvider().IsAvailable())
+			{
+				static const FSlateBrush* ErrorBrush = FAppStyle::Get().GetBrush("SourceControl.StatusIcon.Error");
+				return ErrorBrush;
+			}
+			else
+			{
+				static const FSlateBrush* OnBrush = FAppStyle::Get().GetBrush("SourceControl.StatusIcon.On");
+				return OnBrush;
+			}
+		}
+		else
+		{
+			static const FSlateBrush* OffBrush = FAppStyle::Get().GetBrush("SourceControl.StatusIcon.Off");
+			return OffBrush;
+		}
+	}
+}
+
+TSharedRef<SWidget> FSourceControlMenuHelpers::MakeSourceControlStatusWidget()
+{
+	return
+		SNew(SComboButton)
+		.ContentPadding(FMargin(6.0f,0.0f))
+		.ToolTipText_Static(&FSourceControlMenuHelpers::GetSourceControlTooltip)
+		.MenuPlacement(MenuPlacement_AboveAnchor)
+		.ComboButtonStyle(&FAppStyle::Get().GetWidgetStyle<FComboButtonStyle>("StatusBar.StatusBarComboButton"))
+		.ButtonContent()
+		[
+			SNew(SHorizontalBox)
+			+ SHorizontalBox::Slot()
+			.AutoWidth()
+			.VAlign(VAlign_Center)
+			.HAlign(HAlign_Center)
+			[
+				SNew(SImage)
+				.Image_Static(&FSourceControlMenuHelpers::GetSourceControlIcon)
+			]
+			+ SHorizontalBox::Slot()
+			.AutoWidth()
+			.VAlign(VAlign_Center)
+			.Padding(FMargin(5, 5, 0, 5))
+			[
+				SNew(STextBlock)
+				.TextStyle(&FAppStyle::Get().GetWidgetStyle<FTextBlockStyle>("NormalText"))
+				.Text_Static(&FSourceControlMenuHelpers::GetSourceControlStatusText)
+			]
+		]
+		.MenuContent()
+		[
+			GenerateSourceControlMenuContent()
+		];
+}
+
+#undef LOCTEXT_NAMESPACE
\ No newline at end of file
diff --git a/Engine/Source/Editor/StatusBar/Private/SourceControlMenuHelpers.h b/Engine/Source/Editor/StatusBar/Private/SourceControlMenuHelpers.h
new file mode 100644
--- /dev/null
+++ b/Engine/Source/Editor/StatusBar/Private/SourceControlMenuHelpers.h
@@ -0,0 +1,57 @@
+// Copyright Epic Games, Inc. All Rights Reserved.
+
+#pragma once
+
+
+class FSourceControlCommands : public TCommands<FSourceControlCommands>
+{
+public:
+	FSourceControlCommands();
+
+	/**
+	 * Initialize commands
+	 */
+	virtual void RegisterCommands() override;
+
+private:
+
+	static void ConnectToSourceControl_Clicked();
+	static bool CheckOutModifiedFiles_CanExecute();
+	static void CheckOutModifiedFiles_Clicked();
+	static bool SubmitToSourceControl_CanExecute();
+	static void SubmitToSourceControl_Clicked();
+
+public:
+	/**
+	 * Source Control Commands
+	 */
+	TSharedPtr< FUICommandInfo > ConnectToSourceControl;
+	TSharedPtr< FUICommandInfo > ChangeSourceControlSettings;
+	TSharedPtr< FUICommandInfo > CheckOutModifiedFiles;
+	TSharedPtr< FUICommandInfo > SubmitToSourceControl;
+
+	static TSharedRef<FUICommandList> ActionList;
+};
+
+
+struct FSourceControlMenuHelpers
+{
+	enum EQueryState
+	{
+		NotQueried,
+		Querying,
+		Queried,
+	};
+
+
+	static EQueryState QueryState;
+
+	static void CheckSourceControlStatus();
+	static void OnSourceControlOperationComplete(const FSourceControlOperationRef& InOperation, ECommandResult::Type InResult);
+	static TSharedRef<SWidget> GenerateSourceControlMenuContent();
+	static FText GetSourceControlStatusText();
+	static FText GetSourceControlTooltip();
+	static const FSlateBrush* GetSourceControlIcon();
+	static TSharedRef<SWidget> MakeSourceControlStatusWidget();
+};
+
diff --git a/Engine/Source/Editor/StatusBar/Private/StatusBarModule.cpp b/Engine/Source/Editor/StatusBar/Private/StatusBarModule.cpp
new file mode 100644
--- /dev/null
+++ b/Engine/Source/Editor/StatusBar/Private/StatusBarModule.cpp
@@ -0,0 +1,29 @@
+// Copyright Epic Games, Inc. All Rights Reserved.
+
+#include "StatusBarModule.h"
+
+
+class FStatusBarModule : public IStatusBarModule
+{
+	/** IModuleInterface implementation */
+	virtual void StartupModule() override;
+	virtual void ShutdownModule() override;
+};
+
+IMPLEMENT_MODULE( FStatusBarModule, StatusBar )
+
+
+
+void FStatusBarModule::StartupModule()
+{
+
+}
+
+
+void FStatusBarModule::ShutdownModule()
+{
+
+}
+
+
+
diff --git a/Engine/Source/Editor/StatusBar/Private/StatusBarSubsystem.cpp b/Engine/Source/Editor/StatusBar/Private/StatusBarSubsystem.cpp
new file mode 100644
--- /dev/null
+++ b/Engine/Source/Editor/StatusBar/Private/StatusBarSubsystem.cpp
@@ -0,0 +1,208 @@
+// Copyright Epic Games, Inc. All Rights Reserved.
+
+#pragma once
+
+#include "StatusBarSubsystem.h"
+#include "ToolMenus.h"
+
+#include "Framework/Docking/TabManager.h"
+
+#include "Framework/Application/SlateApplication.h"
+#include "ContentBrowserModule.h"
+#include "IContentBrowserSingleton.h"
+
+
+
+#define LOCTEXT_NAMESPACE "StatusBar"
+
+
+int32 UStatusBarSubsystem::HandleCounter = 0;
+
+
+void UStatusBarSubsystem::Initialize(FSubsystemCollectionBase& Collection)
+{
+	FSourceControlCommands::Register();
+}
+
+void UStatusBarSubsystem::Deinitialize()
+{
+	FSourceControlCommands::Unregister();
+}
+
+bool UStatusBarSubsystem::FocusDebugConsole(TSharedRef<SWindow> ParentWindow)
+{
+	bool bFocusedSuccessfully = false;
+
+	for (auto StatusBar : StatusBars)
+	{
+		if (TSharedPtr<SStatusBar> StatusBarPinned = StatusBar.Value.Pin())
+		{
+			TSharedPtr<SDockTab> ParentTab = StatusBarPinned->GetParentTab();
+			if (ParentTab && ParentTab->IsForeground() && ParentTab->GetParentWindow() == ParentWindow)
+			{
+				// Cache off the previously focused widget so we can restore focus if the user hits the focus key again
+				PreviousKeyboardFocusedWidget = FSlateApplication::Get().GetKeyboardFocusedWidget();
+				StatusBarPinned->FocusDebugConsole();
+
+				bFocusedSuccessfully = true;
+				break;
+			}
+		}
+	}
+
+	return bFocusedSuccessfully;
+}
+
+bool UStatusBarSubsystem::OpenContentBrowser(TSharedRef<SWindow> ParentWindow)
+{
+	bool bFocusedSuccessfully = false;
+
+	CreateContentBrowserIfNeeded();
+
+	TSharedPtr<SWindow> Window = ParentWindow;
+	GEditor->GetTimerManager()->SetTimerForNextTick(FTimerDelegate::CreateUObject(this, &UStatusBarSubsystem::HandleDeferredOpenContentBrowser, Window));
+
+	return true;
+}
+
+TSharedRef<SWidget> UStatusBarSubsystem::MakeStatusBarWidget(FName StatusBarName, const TSharedRef<SDockTab>& InParentTab)
+{
+	TSharedRef<SStatusBar> StatusBar =
+		SNew(SStatusBar, StatusBarName, InParentTab)
+		.OnConsoleClosed_UObject(this, &UStatusBarSubsystem::OnDebugConsoleClosed)
+		.OnGetContentBrowser_UObject(this, &UStatusBarSubsystem::OnGetContentBrowser)
+		.OnContentBrowserOpened_UObject(this, &UStatusBarSubsystem::OnContentBrowserOpened)
+		.OnContentBrowserDismissed_UObject(this, &UStatusBarSubsystem::OnContentBrowserDismissed);
+
+	// Clean up stale status bars
+	for (auto It = StatusBars.CreateIterator(); It; ++It)
+	{
+		if (!It.Value().IsValid())
+		{
+			It.RemoveCurrent();
+		}
+	}
+
+	StatusBars.Add(StatusBarName, StatusBar);
+
+	return StatusBar;
+}
+
+FStatusBarMessageHandle UStatusBarSubsystem::PushStatusBarMessage(FName StatusBarName, const TAttribute<FText>& InMessage, const TAttribute<FText>& InHintText)
+{
+	if (TSharedPtr<SStatusBar> StatusBar = GetStatusBar(StatusBarName))
+	{
+		FStatusBarMessageHandle NewHandle(++HandleCounter);
+
+		StatusBar->PushMessage(NewHandle, InMessage, InHintText);
+
+		return NewHandle;
+	}
+
+	return FStatusBarMessageHandle();
+}
+
+FStatusBarMessageHandle UStatusBarSubsystem::PushStatusBarMessage(FName StatusBarName, const TAttribute<FText>& InMessage)
+{
+	return PushStatusBarMessage(StatusBarName, InMessage, TAttribute<FText>());
+}
+
+void UStatusBarSubsystem::PopStatusBarMessage(FName StatusBarName, FStatusBarMessageHandle InHandle)
+{
+	if (TSharedPtr<SStatusBar> StatusBar = GetStatusBar(StatusBarName))
+	{
+		StatusBar->PopMessage(InHandle);
+	}
+}
+
+void UStatusBarSubsystem::ClearStatusBarMessages(FName StatusBarName)
+{
+	if (TSharedPtr<SStatusBar> StatusBar = GetStatusBar(StatusBarName))
+	{
+		StatusBar->ClearAllMessages();
+	}
+}
+
+void UStatusBarSubsystem::OnDebugConsoleClosed()
+{
+	if (PreviousKeyboardFocusedWidget.IsValid())
+	{
+		FSlateApplication::Get().SetKeyboardFocus(PreviousKeyboardFocusedWidget.Pin());
+		PreviousKeyboardFocusedWidget.Reset();
+	}
+}
+
+void UStatusBarSubsystem::CreateContentBrowserIfNeeded()
+{
+	if(!StatusBarContentBrowser.IsValid())
+	{
+		IContentBrowserSingleton& ContentBrowserSingleton = FModuleManager::Get().LoadModuleChecked<FContentBrowserModule>("ContentBrowser").Get();;
+
+		FContentBrowserConfig Config;
+		Config.bCanSetAsPrimaryBrowser = false;
+
+		StatusBarContentBrowser = ContentBrowserSingleton.CreateContentBrowser("StatusBarContentBrowser", nullptr, &Config);
+	}
+}
+
+TSharedPtr<SStatusBar> UStatusBarSubsystem::GetStatusBar(FName StatusBarName) const
+{
+	return StatusBars.FindRef(StatusBarName).Pin();
+}
+
+TSharedRef<SWidget> UStatusBarSubsystem::OnGetContentBrowser()
+{
+	CreateContentBrowserIfNeeded();
+
+	return StatusBarContentBrowser.ToSharedRef();
+}
+
+void UStatusBarSubsystem::OnContentBrowserOpened(TSharedRef<SStatusBar>& StatusBarWithContentBrowser)
+{
+	// Dismiss any other content browser that is opened when one status bar opens it.  The content browser is a shared resource and shouldn't be in the layout twice
+	for (auto StatusBar : StatusBars)
+	{
+		if (TSharedPtr<SStatusBar> StatusBarPinned = StatusBar.Value.Pin())
+		{
+			if (StatusBarWithContentBrowser != StatusBarPinned && StatusBarPinned->IsContentBrowserOpened())
+			{
+				StatusBarPinned->RemoveContentBrowser();
+			}
+		}
+	}
+}
+
+void UStatusBarSubsystem::OnContentBrowserDismissed(const TSharedPtr<SWidget>& NewlyFocusedWidget)
+{
+	if (PreviousKeyboardFocusedWidget.IsValid() && !NewlyFocusedWidget.IsValid())
+	{
+		FSlateApplication::Get().SetKeyboardFocus(PreviousKeyboardFocusedWidget.Pin());
+		
+	}
+
+	PreviousKeyboardFocusedWidget.Reset();
+}
+
+void UStatusBarSubsystem::HandleDeferredOpenContentBrowser(TSharedPtr<SWindow> ParentWindow)
+{
+	for (auto StatusBar : StatusBars)
+	{
+		if (TSharedPtr<SStatusBar> StatusBarPinned = StatusBar.Value.Pin())
+		{
+			TSharedPtr<SDockTab> ParentTab = StatusBarPinned->GetParentTab();
+			if (ParentTab && ParentTab->IsForeground() && ParentTab->GetParentWindow() == ParentWindow)
+			{
+				StatusBarPinned->OpenContentBrowser();
+				IContentBrowserSingleton& ContentBrowserSingleton = FModuleManager::Get().GetModuleChecked<FContentBrowserModule>("ContentBrowser").Get();
+
+				// Cache off the previously focused widget so we can restore focus if the user hits the focus key again
+				PreviousKeyboardFocusedWidget = FSlateApplication::Get().GetKeyboardFocusedWidget();
+
+				ContentBrowserSingleton.FocusContentBrowserSearchField(StatusBarContentBrowser.ToSharedRef());
+				break;
+			}
+		}
+	}
+}
+
+#undef LOCTEXT_NAMESPACE
\ No newline at end of file
diff --git a/Engine/Source/Editor/StatusBar/Public/SStatusBar.h b/Engine/Source/Editor/StatusBar/Public/SStatusBar.h
new file mode 100644
--- /dev/null
+++ b/Engine/Source/Editor/StatusBar/Public/SStatusBar.h
@@ -0,0 +1,114 @@
+// Copyright Epic Games, Inc. All Rights Reserved.
+
+#pragma once
+
+#include "StatusBarSubsystem.h"
+
+class SStatusBar;
+
+DECLARE_DELEGATE_OneParam(FOnContentBrowserOpened, TSharedRef<SStatusBar>&)
+DECLARE_DELEGATE_OneParam(FOnContentBrowserDismissed, const TSharedPtr<SWidget>&)
+
+class SWindow;
+class SDockTab;
+class SWidget;
+class SMultiLineEditableTextBox;
+class SContentBrowserOverlay;
+
+class SStatusBar : public SCompoundWidget
+{
+	SLATE_BEGIN_ARGS(SStatusBar)	
+	{}
+		SLATE_EVENT(FSimpleDelegate, OnConsoleClosed)
+
+		SLATE_EVENT(FOnGetContent, OnGetContentBrowser)
+
+		SLATE_EVENT(FOnContentBrowserOpened, OnContentBrowserOpened)
+
+		SLATE_EVENT(FOnContentBrowserDismissed, OnContentBrowserDismissed)
+
+	SLATE_END_ARGS()
+
+public:
+	virtual bool SupportsKeyboardFocus() const { return false; }
+
+	void Construct(const FArguments& InArgs, FName InStatusBarName, const TSharedRef<SDockTab> InParentTab);
+
+	void PushMessage(FStatusBarMessageHandle& InHandle, const TAttribute<FText>& InMessage, const TAttribute<FText>& InHintText);
+	void PopMessage(FStatusBarMessageHandle& InHandle);
+	void ClearAllMessages();
+
+	EVisibility GetHelpIconVisibility() const;
+
+	TSharedPtr<SDockTab> GetParentTab() const;
+	void FocusDebugConsole();
+	bool IsDebugConsoleFocused() const;
+
+
+	void OnGlobalFocusChanging(const FFocusEvent& FocusEvent, const FWeakWidgetPath& OldFocusedWidgetPath, const TSharedPtr<SWidget>& OldFocusedWidget, const FWidgetPath& NewFocusedWidgetPath, const TSharedPtr<SWidget>& NewFocusedWidget);
+
+	void OnActiveTabChanged(TSharedPtr<SDockTab> PreviouslyActive, TSharedPtr<SDockTab> NewlyActivated);
+
+	bool IsContentBrowserOpened() const;
+
+	void OpenContentBrowser();
+
+	/**
+	 * Dismisses an open content browser with an animation.  The content browser is removed once the animation is complete
+	 */
+	void DismissContentBrowser(const TSharedPtr<SWidget>& NewlyFocusedWidget);
+
+	/**
+	 * Removes a content browser from the layout immediately
+	 */
+	void RemoveContentBrowser();
+
+
+private:
+	const FSlateBrush* GetContentBrowserExpandArrowImage() const;
+
+	FText GetStatusBarMessage() const;
+
+	TSharedRef<SWidget> MakeContentBrowserWidget();
+
+	TSharedRef<SWidget> MakeStatusBarToolBarWidget();
+	TSharedRef<SWidget> MakeDebugConsoleWidget(FSimpleDelegate OnConsoleClosed);
+	TSharedRef<SWidget> MakeStatusMessageWidget();
+
+	FReply OnContentBrowserButtonClicked();
+
+	EActiveTimerReturnType UpdateContentBrowserAnimation(double CurrentTime, float DeltaTime);
+
+	void RegisterStatusBarMenu();
+
+	void RegisterSourceControlStatus();
+
+	struct FStatusBarMessage
+	{
+		FStatusBarMessage(const TAttribute<FText>& InMessageText, const TAttribute<FText>& InHintText, FStatusBarMessageHandle InHandle)
+			: MessageText(InMessageText)
+			, HintText(InHintText)
+			, Handle(InHandle)
+		{}
+
+		TAttribute<FText> MessageText;
+		TAttribute<FText> HintText;
+		FStatusBarMessageHandle Handle;
+	};
+private:
+	TArray<FStatusBarMessage> MessageStack;
+	TSharedPtr<SMultiLineEditableTextBox> ConsoleEditBox;
+	TWeakPtr<SDockTab> ParentTab;
+	TSharedPtr<SContentBrowserOverlay> ContentBrowserOverlayContent;
+	TSharedPtr<FActiveTimerHandle> ContentBrowserOpenCloseTimer;
+	FOnGetContent GetContentBrowserDelegate;
+	FOnContentBrowserOpened OnContentBrowserOpenedDelegate;
+	FOnContentBrowserDismissed OnContentBrowserDismissedDelegate;
+	FCurveSequence ContentBrowserEasingCurve;
+	float MaxContentBrowserHeight;
+	float MinContentBrowserHeight;
+	const FSlateBrush* UpArrow;
+	const FSlateBrush* DownArrow;
+	FName StatusBarName;
+};
+
diff --git a/Engine/Source/Editor/StatusBar/Public/StatusBarModule.h b/Engine/Source/Editor/StatusBar/Public/StatusBarModule.h
new file mode 100644
--- /dev/null
+++ b/Engine/Source/Editor/StatusBar/Public/StatusBarModule.h
@@ -0,0 +1,27 @@
+// Copyright Epic Games, Inc. All Rights Reserved.
+
+#pragma once
+
+#include "CoreMinimal.h"
+#include "Modules/ModuleInterface.h"
+#include "Modules/ModuleManager.h"
+
+
+/**
+ * The public interface to this module
+ */
+class IStatusBarModule : public IModuleInterface
+{
+
+public:
+	static inline IStatusBarModule& Get()
+	{
+		return FModuleManager::LoadModuleChecked<IStatusBarModule>( "StatusBar" );
+	}
+
+	static inline bool IsAvailable()
+	{
+		return FModuleManager::Get().IsModuleLoaded("StatusBar");
+	}
+};
+
diff --git a/Engine/Source/Editor/StatusBar/Public/StatusBarSubsystem.h b/Engine/Source/Editor/StatusBar/Public/StatusBarSubsystem.h
new file mode 100644
--- /dev/null
+++ b/Engine/Source/Editor/StatusBar/Public/StatusBarSubsystem.h
@@ -0,0 +1,126 @@
+// Copyright Epic Games, Inc. All Rights Reserved.
+
+#pragma once
+
+#include "CoreMinimal.h"
+#include "EditorSubsystem.h"
+
+#include "StatusBarSubsystem.generated.h"
+
+
+class SStatusBar;
+
+
+template<typename ObjectType> 
+class TAttribute;
+
+struct FStatusBarMessageHandle
+{
+	friend class UStatusBarSubsystem;
+
+	FStatusBarMessageHandle()
+		: Id(INDEX_NONE)
+	{}
+
+	bool IsValid()
+	{
+		return Id != INDEX_NONE;
+	}
+
+	void Reset()
+	{
+		Id = INDEX_NONE;
+	}
+
+	bool operator==(const FStatusBarMessageHandle& OtherHandle) const
+	{
+		return Id == OtherHandle.Id;
+	}
+private:
+	FStatusBarMessageHandle(int32 InId)
+		: Id(InId)
+	{}
+
+	int32 Id;
+};
+
+UCLASS()
+class STATUSBAR_API UStatusBarSubsystem : public UEditorSubsystem
+{
+	GENERATED_BODY()
+
+public:
+
+	/**
+	 *	Prepares for use
+	 */
+	virtual void Initialize(FSubsystemCollectionBase& Collection) override final;
+
+	/**
+	 *	Internal cleanup
+	 */
+	virtual void Deinitialize() override final;
+
+	/**
+	 * Focuses the debug console on the status bar for status bar residing in the passed in parent window 
+	 *
+	 * @param ParentWindow The parent window of the status bar 
+	 * @return true of a status bar debug console was found and focused 
+	 */
+	bool FocusDebugConsole(TSharedRef<SWindow> ParentWindow);
+
+	/**
+	 * Opens the content browser drawer for a status bar residing in the passed in parent window
+	 */
+	bool OpenContentBrowser(TSharedRef<SWindow> ParentWindow);
+
+	/** 
+	 * Creates a new instance of a status bar widget
+	 *
+	 * @param StatusBarName	The name of the status bar for updating it later.
+	 */
+	TSharedRef<SWidget> MakeStatusBarWidget(FName StatusBarName, const TSharedRef<SDockTab>& InParentTab);
+
+	/** 
+	 * Pushes a new status bar message
+	 *
+	 * @param StatusBarName	The name of the status bar to push messages to
+	 * @param InMessage		The message to display
+	 * @param InHintText	Optional hint text message.  This message will be highlighted to make it stand out
+	 * @return	A handle to the message for clearing it later
+	 */
+	FStatusBarMessageHandle PushStatusBarMessage(FName StatusBarName, const TAttribute<FText>& InMessage, const TAttribute<FText>& InHintText);
+	FStatusBarMessageHandle PushStatusBarMessage(FName StatusBarName, const TAttribute<FText>& InMessage);
+
+	/**
+	 * Removes a message from the status bar.  When messages are removed the previous message on the stack (if any) is displayed
+	 *
+	 * @param StatusBarName	The name of the status bar to remove from
+	 * @param InHandle		Handle to the status bar message to remove
+	 */
+	void PopStatusBarMessage(FName StatusBarName, FStatusBarMessageHandle InHandle);
+
+	/**
+	 * Removes all message from the status bar
+	 *
+	 * @param StatusBarName	The name of the status bar to remove from
+	 */
+	void ClearStatusBarMessages(FName StatusBarName);
+
+private:
+	void OnDebugConsoleClosed();
+	void CreateContentBrowserIfNeeded();
+	TSharedPtr<SStatusBar> GetStatusBar(FName StatusBarName) const;
+	TSharedRef<SWidget> OnGetContentBrowser();
+	void OnContentBrowserOpened(TSharedRef<SStatusBar>& StatusBarWithContentBrowser);
+	void OnContentBrowserDismissed(const TSharedPtr<SWidget>& NewlyFocusedWidget);
+	void HandleDeferredOpenContentBrowser(TSharedPtr<SWindow> ParentWindow);
+
+private:
+	TMap<FName, TWeakPtr<SStatusBar>> StatusBars;
+	TWeakPtr<SWidget> PreviousKeyboardFocusedWidget;
+
+	/** The floating content browser that is opened via the content browser button in the status bar */
+	TSharedPtr<SWidget> StatusBarContentBrowser;
+	static int32 HandleCounter;
+};
diff --git a/Engine/Source/Editor/StatusBar/StatusBar.Build.cs b/Engine/Source/Editor/StatusBar/StatusBar.Build.cs
new file mode 100644
--- /dev/null
+++ b/Engine/Source/Editor/StatusBar/StatusBar.Build.cs
@@ -0,0 +1,27 @@
+// Copyright Epic Games, Inc. All Rights Reserved.
+
+using UnrealBuildTool;
+
+public class StatusBar : ModuleRules
+{
+    public StatusBar(ReadOnlyTargetRules Target)
+         : base(Target)
+    {
+		PublicIncludePaths.Add(ModuleDirectory + "/Public");
+
+		PrivateDependencyModuleNames.AddRange(
+            new string[] {
+                "Core",
+                "CoreUObject",
+                "Engine",           
+				"Slate",
+                "SlateCore",
+                "EditorStyle",
+                "UnrealEd",
+				"ToolMenus",
+				"OutputLog",
+				"SourceControlWindows",
+				"EditorSubsystem",
+            });
+    }
+}
diff --git a/Engine/Source/Editor/UnrealEd/Private/Subsystems/StatusBarSubsystem.cpp b/Engine/Source/Editor/UnrealEd/Private/Subsystems/StatusBarSubsystem.cpp
deleted file mode 100644
diff --git a/Engine/Source/Editor/UnrealEd/Private/Toolkits/GlobalEditorCommonCommands.cpp b/Engine/Source/Editor/UnrealEd/Private/Toolkits/GlobalEditorCommonCommands.cpp
--- a/Engine/Source/Editor/UnrealEd/Private/Toolkits/GlobalEditorCommonCommands.cpp
+++ b/Engine/Source/Editor/UnrealEd/Private/Toolkits/GlobalEditorCommonCommands.cpp
@@ -14,7 +14,7 @@
 #include "Widgets/Docking/SDockTab.h"
 #include "Toolkits/SGlobalOpenAssetDialog.h"
 #include "Toolkits/SGlobalTabSwitchingDialog.h"
-#include "Subsystems/StatusBarSubsystem.h"
+#include "StatusBarSubsystem.h"
 
 #define LOCTEXT_NAMESPACE "GlobalEditorCommonCommands"
 
@@ -28,32 +28,34 @@ FGlobalEditorCommonCommands::FGlobalEditorCommonCommands()
 
 void FGlobalEditorCommonCommands::RegisterCommands()
 {
-	UI_COMMAND( SummonControlTabNavigation, "Tab Navigation", "Summons a list of open assets and tabs", EUserInterfaceActionType::Button, FInputChord(EModifierKey::Control, EKeys::Tab) );
-	UI_COMMAND( SummonControlTabNavigationAlternate, "Tab Navigation", "Summons a list of open assets and tabs", EUserInterfaceActionType::Button, FInputChord(EModifierKey::Command, EKeys::Tab) );
-
-	UI_COMMAND( SummonOpenAssetDialog, "Open Asset...", "Summons an asset picker", EUserInterfaceActionType::Button, FInputChord(EModifierKey::Control, EKeys::P) );
-	UI_COMMAND( SummonOpenAssetDialogAlternate, "Open Asset...", "Summons an asset picker", EUserInterfaceActionType::Button, FInputChord(EModifierKey::Alt | EModifierKey::Shift, EKeys::O));
-	UI_COMMAND( FindInContentBrowser, "Browse to Asset", "Browses to the associated asset and selects it in the most recently used Content Browser (summoning one if necessary)", EUserInterfaceActionType::Button, FInputChord(EModifierKey::Control, EKeys::B));
-	
-	UI_COMMAND( OpenConsoleCommandBox, "Open Console Command Box", "Opens an edit box where you can type in a console command", EUserInterfaceActionType::Button, FInputChord(EKeys::Tilde));
-	UI_COMMAND( OpenDocumentation, "Open Documentation...", "Opens documentation for this tool", EUserInterfaceActionType::Button, FInputChord(EKeys::F1) );
+	UI_COMMAND(SummonControlTabNavigation, "Tab Navigation", "Summons a list of open assets and tabs", EUserInterfaceActionType::Button, FInputChord(EModifierKey::Control, EKeys::Tab));
+	UI_COMMAND(SummonControlTabNavigationAlternate, "Tab Navigation", "Summons a list of open assets and tabs", EUserInterfaceActionType::Button, FInputChord(EModifierKey::Command, EKeys::Tab));
+
+	UI_COMMAND(SummonOpenAssetDialog, "Open Asset...", "Summons an asset picker", EUserInterfaceActionType::Button, FInputChord(EModifierKey::Control, EKeys::P));
+	UI_COMMAND(SummonOpenAssetDialogAlternate, "Open Asset...", "Summons an asset picker", EUserInterfaceActionType::Button, FInputChord(EModifierKey::Alt | EModifierKey::Shift, EKeys::O));
+	UI_COMMAND(FindInContentBrowser, "Browse to Asset", "Browses to the associated asset and selects it in the most recently used Content Browser (summoning one if necessary)", EUserInterfaceActionType::Button, FInputChord(EModifierKey::Control, EKeys::B));
+
+	UI_COMMAND(OpenConsoleCommandBox, "Open Console Command Box", "Opens an edit box where you can type in a console command", EUserInterfaceActionType::Button, FInputChord(EKeys::Tilde));
+	UI_COMMAND(OpenDocumentation, "Open Documentation...", "Opens documentation for this tool", EUserInterfaceActionType::Button, FInputChord(EKeys::F1));
+
+	UI_COMMAND(OpenContentBrowserDrawer, "Open Content Browser Drawer", "Opens the content browser drawer from the status bar and focuses the search field", EUserInterfaceActionType::Button, FInputChord(EModifierKey::Control, EKeys::SpaceBar));
 }
 
 void FGlobalEditorCommonCommands::MapActions(TSharedRef<FUICommandList>& ToolkitCommands)
 {
 	Register();
 
- 	ToolkitCommands->MapAction(
- 		Get().SummonControlTabNavigation,
- 		FExecuteAction::CreateStatic( &FGlobalEditorCommonCommands::OnPressedCtrlTab, Get().SummonControlTabNavigation ) );
+	ToolkitCommands->MapAction(
+		Get().SummonControlTabNavigation,
+		FExecuteAction::CreateStatic(&FGlobalEditorCommonCommands::OnPressedCtrlTab, Get().SummonControlTabNavigation));
 
- 	ToolkitCommands->MapAction(
- 		Get().SummonControlTabNavigationAlternate,
- 		FExecuteAction::CreateStatic( &FGlobalEditorCommonCommands::OnPressedCtrlTab, Get().SummonControlTabNavigationAlternate ) );
+	ToolkitCommands->MapAction(
+		Get().SummonControlTabNavigationAlternate,
+		FExecuteAction::CreateStatic(&FGlobalEditorCommonCommands::OnPressedCtrlTab, Get().SummonControlTabNavigationAlternate));
 
 	ToolkitCommands->MapAction(
 		Get().SummonOpenAssetDialog,
-		FExecuteAction::CreateStatic( &FGlobalEditorCommonCommands::OnSummonedAssetPicker ) );
+		FExecuteAction::CreateStatic(&FGlobalEditorCommonCommands::OnSummonedAssetPicker));
 
 	ToolkitCommands->MapAction(
 		Get().SummonOpenAssetDialogAlternate,
@@ -61,7 +63,11 @@ void FGlobalEditorCommonCommands::MapActions(TSharedRef<FUICommandList>& Toolkit
 
 	ToolkitCommands->MapAction(
 		Get().OpenConsoleCommandBox,
-		FExecuteAction::CreateStatic( &FGlobalEditorCommonCommands::OnSummonedConsoleCommandBox ) );
+		FExecuteAction::CreateStatic(&FGlobalEditorCommonCommands::OnSummonedConsoleCommandBox));
+
+	ToolkitCommands->MapAction(
+		Get().OpenContentBrowserDrawer,
+		FExecuteAction::CreateStatic(&FGlobalEditorCommonCommands::OnOpenContentBrowserDrawer));
 }
 
 void FGlobalEditorCommonCommands::OnPressedCtrlTab(TSharedPtr<FUICommandInfo> TriggeringCommand)
@@ -128,8 +134,8 @@ static void CloseDebugConsole()
 void FGlobalEditorCommonCommands::OnSummonedConsoleCommandBox()
 {
 	TSharedPtr<SWindow> ParentWindow = FSlateApplication::Get().GetActiveTopLevelWindow();
-	
-	if( ParentWindow.IsValid() && ParentWindow->GetType() == EWindowType::Normal )
+
+	if (ParentWindow.IsValid() && ParentWindow->GetType() == EWindowType::Normal)
 	{
 		TSharedRef<SWindow> WindowRef = ParentWindow.ToSharedRef();
 		FOutputLogModule& OutputLogModule = FModuleManager::LoadModuleChecked<FOutputLogModule>(TEXT("OutputLog"));
@@ -146,6 +152,29 @@ void FGlobalEditorCommonCommands::OnSummonedConsoleCommandBox()
 	}
 }
 
+void FGlobalEditorCommonCommands::OnOpenContentBrowserDrawer()
+{
+	TSharedPtr<SWindow> ParentWindow = FSlateApplication::Get().GetActiveTopLevelWindow();
+	if (!ParentWindow.IsValid())
+	{
+		if (TSharedPtr<SDockTab> ActiveTab = FGlobalTabmanager::Get()->GetActiveTab())
+		{
+			if (TSharedPtr<SDockTab> ActiveMajorTab = FGlobalTabmanager::Get()->GetMajorTabForTabManager(ActiveTab->GetTabManager()))
+			{
+				ParentWindow = ActiveMajorTab->GetParentWindow();
+			}
+		}
+	}
+
+	if (ParentWindow.IsValid() && ParentWindow->GetType() == EWindowType::Normal)
+	{
+		TSharedRef<SWindow> WindowRef = ParentWindow.ToSharedRef();
+		FOutputLogModule& OutputLogModule = FModuleManager::LoadModuleChecked<FOutputLogModule>(TEXT("OutputLog"));
+
+		GEditor->GetEditorSubsystem<UStatusBarSubsystem>()->OpenContentBrowser(ParentWindow.ToSharedRef());
+	}
+}
+
 //////////////////////////////////////////////////////////////////////////
 
 #undef LOCTEXT_NAMESPACE
diff --git a/Engine/Source/Editor/UnrealEd/Private/Toolkits/SStandaloneAssetEditorToolkitHost.cpp b/Engine/Source/Editor/UnrealEd/Private/Toolkits/SStandaloneAssetEditorToolkitHost.cpp
--- a/Engine/Source/Editor/UnrealEd/Private/Toolkits/SStandaloneAssetEditorToolkitHost.cpp
+++ b/Engine/Source/Editor/UnrealEd/Private/Toolkits/SStandaloneAssetEditorToolkitHost.cpp
@@ -12,7 +12,7 @@
 #include "Interfaces/IMainFrameModule.h"
 #include "Widgets/Docking/SDockTab.h"
 #include "UObject/Package.h"
-#include "Subsystems/StatusBarSubsystem.h"
+#include "StatusBarSubsystem.h"
 #include "ToolMenus.h"
 #include "IIntroTutorials.h"
 
diff --git a/Engine/Source/Editor/UnrealEd/Public/Subsystems/StatusBarSubsystem.h b/Engine/Source/Editor/UnrealEd/Public/Subsystems/StatusBarSubsystem.h
deleted file mode 100644
diff --git a/Engine/Source/Editor/UnrealEd/Public/Toolkits/GlobalEditorCommonCommands.h b/Engine/Source/Editor/UnrealEd/Public/Toolkits/GlobalEditorCommonCommands.h
--- a/Engine/Source/Editor/UnrealEd/Public/Toolkits/GlobalEditorCommonCommands.h
+++ b/Engine/Source/Editor/UnrealEd/Public/Toolkits/GlobalEditorCommonCommands.h
@@ -25,6 +25,7 @@ protected:
 	static void OnPressedCtrlTab(TSharedPtr<FUICommandInfo> TriggeringCommand);
 	static void OnSummonedAssetPicker();
 	static void OnSummonedConsoleCommandBox();
+	static void OnOpenContentBrowserDrawer();
 
 	static TSharedPtr<IMenu> OpenPopupMenu(TSharedRef<SWidget> WindowContents, const FVector2D& PopupDesiredSize);
 public:
@@ -35,5 +36,6 @@ public:
 	TSharedPtr<FUICommandInfo> SummonOpenAssetDialogAlternate;
 	TSharedPtr<FUICommandInfo> OpenDocumentation;
 	TSharedPtr<FUICommandInfo> OpenConsoleCommandBox;
+	TSharedPtr<FUICommandInfo> OpenContentBrowserDrawer;
 };
 
diff --git a/Engine/Source/Editor/UnrealEd/UnrealEd.Build.cs b/Engine/Source/Editor/UnrealEd/UnrealEd.Build.cs
--- a/Engine/Source/Editor/UnrealEd/UnrealEd.Build.cs
+++ b/Engine/Source/Editor/UnrealEd/UnrealEd.Build.cs
@@ -108,7 +108,8 @@ public class UnrealEd : ModuleRules
                 "EditorSubsystem",
                 "InteractiveToolsFramework",
 				"ToolMenusEditor",
-            }
+				"StatusBar",
+			}
 		);
 
 		PrivateDependencyModuleNames.AddRange(
diff --git a/Engine/Source/Runtime/Slate/Private/Framework/Application/MenuStack.cpp b/Engine/Source/Runtime/Slate/Private/Framework/Application/MenuStack.cpp
--- a/Engine/Source/Runtime/Slate/Private/Framework/Application/MenuStack.cpp
+++ b/Engine/Source/Runtime/Slate/Private/Framework/Application/MenuStack.cpp
@@ -8,6 +8,8 @@
 #include "Widgets/Layout/SBox.h"
 #include "Widgets/Layout/SPopup.h"
 #include "Framework/Application/Menu.h"
+#include "Widgets/Layout/SBorder.h"
+#include "Widgets/Images/SImage.h"
 
 #define LOCTEXT_NAMESPACE "MenuStack"
 
@@ -103,7 +105,28 @@ namespace MenuStackInternal
 				.MinDesiredWidth(InArgs._OptionalMinMenuWidth)
 				.MaxDesiredHeight(InArgs._OptionalMinMenuHeight)
 				[
-				InArgs._MenuContent.Widget
+					// Always add a background to the menu.  This includes a small outline around the background to distinguish open menus from each other
+					SNew(SOverlay)
+					+ SOverlay::Slot()
+					[
+						SNew(SImage)
+						.Image(FCoreStyle::Get().GetBrush("Menu.Background"))
+					]
+					+ SOverlay::Slot()
+					[
+						SNew(SImage)
+						.Image(FCoreStyle::Get().GetOptionalBrush("Menu.Outline", nullptr))
+					]
+					+ SOverlay::Slot()
+					[
+						SNew(SBorder)
+						.Padding(0)
+						.BorderImage(FStyleDefaults::GetNoBrush())
+						.ForegroundColor(FCoreStyle::Get().GetSlateColor("DefaultForeground"))
+						[
+							InArgs._MenuContent.Widget
+						]
+					]
 				]
 			];
 		}
@@ -588,13 +611,15 @@ void FMenuStack::SetHostPath(const FWidgetPath& InOwnerPath)
 
 	HostWindow = InOwnerPath.IsValid() ? InOwnerPath.GetWindow() : TSharedPtr<SWindow>();
 
+	HostWidget = InOwnerPath.IsValid() ? InOwnerPath.GetLastWidget() : TWeakPtr<SWidget>();
+
 	if ( HostWindow.IsValid() && !HostWindowPopupPanel.IsValid() )
 	{
 		TSharedRef<SMenuPanel> NewHostWindowPopupPanel = SNew(SMenuPanel);
 		for ( int i = InOwnerPath.Widgets.Num() - 1; i >= 0; i-- )
 		{
-			const TSharedRef<SWidget>& HostWidget = InOwnerPath.Widgets[i].Widget;
-			HostPopupLayer = HostWidget->OnVisualizePopup(NewHostWindowPopupPanel);
+			const TSharedRef<SWidget>& CurrentWidget = InOwnerPath.Widgets[i].Widget;
+			HostPopupLayer = CurrentWidget->OnVisualizePopup(NewHostWindowPopupPanel);
 			if ( HostPopupLayer.IsValid() )
 			{
 				HostWindowPopupPanel = NewHostWindowPopupPanel;
@@ -818,4 +843,9 @@ TSharedPtr<SWindow> FMenuStack::GetHostWindow() const
 	return HostWindow;
 }
 
+TSharedPtr<SWidget> FMenuStack::GetHostWidget() const
+{
+	return HostWidget.Pin();
+}
+
 #undef LOCTEXT_NAMESPACE
diff --git a/Engine/Source/Runtime/Slate/Private/Framework/Application/SlateApplication.cpp b/Engine/Source/Runtime/Slate/Private/Framework/Application/SlateApplication.cpp
--- a/Engine/Source/Runtime/Slate/Private/Framework/Application/SlateApplication.cpp
+++ b/Engine/Source/Runtime/Slate/Private/Framework/Application/SlateApplication.cpp
@@ -41,6 +41,7 @@
 #include "Math/UnitConversion.h"
 #include "HAL/LowLevelMemTracker.h"
 #include "ProfilingDebugging/CsvProfiler.h"
+#include "Styling/StarshipCoreStyle.h"
 
 #ifndef SLATE_HAS_WIDGET_REFLECTOR
 	#define SLATE_HAS_WIDGET_REFLECTOR !(UE_BUILD_TEST || UE_BUILD_SHIPPING) && PLATFORM_DESKTOP
@@ -51,6 +52,7 @@
 #endif
 #include "Debugging/SlateDebugging.h"
 
+
 CSV_DECLARE_CATEGORY_MODULE_EXTERN(CORE_API, Basic);
 
 CSV_DECLARE_CATEGORY_MODULE_EXTERN(SLATECORE_API, Slate);
@@ -633,8 +635,9 @@ TSharedRef<FSlateApplication> FSlateApplication::Create(const TSharedRef<class G
 {
 	EKeys::Initialize();
 
-	FCoreStyle::ResetToDefault();
+	InitializeCoreStyle();
 
+	
 	// Note: Important to establish the static PlatformApplication property first, as the FSlateApplication ctor relies on it
 	PlatformApplication = InPlatformApplication;
 
@@ -717,7 +720,6 @@ FSlateApplication::FSlateApplication()
 	, bTouchFallbackToMouse( true )
 	, bSoftwareCursorAvailable( false )	
 	, bMenuAnimationsEnabled( false )
-	, AppIcon( FCoreStyle::Get().GetBrush("DefaultAppIcon") )
 	, VirtualDesktopRect( 0,0,0,0 )
 	, NavigationConfig(MakeShared<FNavigationConfig>())
 #if WITH_EDITOR
@@ -2071,6 +2073,11 @@ TSharedPtr<SWindow> FSlateApplication::GetVisibleMenuWindow() const
 	return MenuStack.GetHostWindow();
 }
 
+TSharedPtr<SWidget> FSlateApplication::GetMenuHostWidget() const
+{
+	return MenuStack.GetHostWidget();
+}
+
 void FSlateApplication::DismissAllMenus()
 {
 	MenuStack.DismissAll();
@@ -3621,9 +3628,15 @@ void FSlateApplication::SetAppIcon(const FSlateBrush* const InAppIcon)
 
 const FSlateBrush* FSlateApplication::GetAppIcon() const
 {
-	return AppIcon;
+	static FName AppIconName("AppIcon");
+	return FAppStyle::Get().GetBrush(AppIconName);
 }
 
+const FSlateBrush* FSlateApplication::GetAppIconSmall() const
+{
+	static FName AppIconName("AppIcon.Small");
+	return FAppStyle::Get().GetBrush(AppIconName);
+}
 
 void FSlateApplication::ShowVirtualKeyboard( bool bShow, int32 UserIndex, TSharedPtr<IVirtualKeyboardEntry> TextEntryWidget )
 {
@@ -6544,6 +6557,21 @@ TSharedRef<FSlateApplication> FSlateApplication::InitializeAsStandaloneApplicati
 	return Slate;
 }
 
+void FSlateApplication::InitializeCoreStyle()
+{
+	if (FCoreStyle::IsStarshipStyle())
+	{
+		FStarshipCoreStyle::ResetToDefault();
+		FAppStyle::SetAppStyleSet(FStarshipCoreStyle::GetCoreStyle());
+	}
+	else
+	{
+		FCoreStyle::ResetToDefault();
+		FAppStyle::SetAppStyleSet(FCoreStyle::GetCoreStyle());
+	}
+
+}
+
 void FSlateApplication::SetWidgetReflector(const TSharedRef<IWidgetReflector>& WidgetReflector)
 {
 	if ( SourceCodeAccessDelegate.IsBound() )
diff --git a/Engine/Source/Runtime/Slate/Public/Framework/Application/MenuStack.h b/Engine/Source/Runtime/Slate/Public/Framework/Application/MenuStack.h
--- a/Engine/Source/Runtime/Slate/Public/Framework/Application/MenuStack.h
+++ b/Engine/Source/Runtime/Slate/Public/Framework/Application/MenuStack.h
@@ -189,6 +189,11 @@ public:
 	 */
 	TSharedPtr<SWindow> GetHostWindow() const;
 
+	/**
+	 * @return	Returns the widget that is the parent widget that initially created the menu in the stack
+	 */
+	TSharedPtr<SWidget> GetHostWidget() const;
+
 	/**
 	* @return	True if the stack has one or more menus in it, false if it is empty.
 	*/
@@ -334,6 +339,9 @@ private:
 	/** The parent window of the root menu in the stack. NOT the actual menu window if it's a CreateNewWindow */
 	TSharedPtr<SWindow> HostWindow;
 
+	/** The parent widget of the root menu in the stack */
+	TWeakPtr<SWidget> HostWidget;
+
 	/** The parent window of the root menu in the stack. NOT the actual menu window if it's a CreateNewWindow */
 	TSharedPtr<SMenuPanel> HostWindowPopupPanel;
 
diff --git a/Engine/Source/Runtime/Slate/Public/Framework/Application/SlateApplication.h b/Engine/Source/Runtime/Slate/Public/Framework/Application/SlateApplication.h
--- a/Engine/Source/Runtime/Slate/Public/Framework/Application/SlateApplication.h
+++ b/Engine/Source/Runtime/Slate/Public/Framework/Application/SlateApplication.h
@@ -202,6 +202,8 @@ public:
 	
 	static TSharedRef<FSlateApplication> InitializeAsStandaloneApplication(const TSharedRef< class FSlateRenderer >& PlatformRenderer, const TSharedRef<class GenericApplication>& InPlatformApplication);
 
+	static void InitializeCoreStyle();
+
 	/**
 	 * Returns true if a Slate application instance is currently initialized and ready
 	 *
@@ -405,6 +407,9 @@ public:
 	/** @return	Returns a ptr to the window that is currently the host of the menu stack or null if no menus are visible */
 	TSharedPtr<SWindow> GetVisibleMenuWindow() const;
 
+	/** @return	Returns a ptr to the widget that created the opened root menu or null if one is not opened */
+	TSharedPtr<SWidget> GetMenuHostWidget() const;
+
 	/** Dismisses all open menus */
 	void DismissAllMenus();
 
@@ -819,6 +824,7 @@ public:
 	void OverridePlatformApplication(TSharedPtr<class GenericApplication> InPlatformApplication);
 
 	/** Set the global application icon */
+	UE_DEPRECATED(4.26, "SetAppIcon has been deprecated.  Set \"AppIcon\" in your applications style to override the icon")
 	void SetAppIcon(const FSlateBrush* const InAppIcon);
 
 	/** Sets the display state of external UI such as Steam. */
@@ -1329,6 +1335,7 @@ public:
 
 	virtual TSharedPtr<SWindow> GetActiveTopLevelWindow() const override;
 	virtual const FSlateBrush* GetAppIcon() const override;
+	virtual const FSlateBrush* GetAppIconSmall() const override;
 
 	virtual float GetApplicationScale() const override { return Scale; }
 	virtual bool GetSoftwareCursorAvailable() const override { return bSoftwareCursorAvailable; }
