Added files that didnt save

diff --git a/Engine/Source/Editor/StatusBar/Private/SStatusBar.cpp b/Engine/Source/Editor/StatusBar/Private/SStatusBar.cpp
--- a/Engine/Source/Editor/StatusBar/Private/SStatusBar.cpp
+++ b/Engine/Source/Editor/StatusBar/Private/SStatusBar.cpp
@@ -19,6 +19,8 @@
 #include "ToolMenuContext.h"
 #include "ToolMenus.h"
 #include "OutputLogModule.h"
+#include "SourceControlMenuHelpers.h"
+#include "Framework/Application/SlateApplication.h"
 
 #define LOCTEXT_NAMESPACE "StatusBar"
 
diff --git a/Engine/Source/Editor/StatusBar/Private/SourceControlMenuHelpers.cpp b/Engine/Source/Editor/StatusBar/Private/SourceControlMenuHelpers.cpp
--- a/Engine/Source/Editor/StatusBar/Private/SourceControlMenuHelpers.cpp
+++ b/Engine/Source/Editor/StatusBar/Private/SourceControlMenuHelpers.cpp
@@ -1,5 +1,6 @@
 // Copyright Epic Games, Inc. All Rights Reserved.
 
+#include "SourceControlMenuHelpers.h"
 #include "ISourceControlOperation.h"
 #include "SourceControlOperations.h"
 #include "ISourceControlProvider.h"
@@ -9,6 +10,9 @@
 #include "FileHelpers.h"
 #include "ToolMenuContext.h"
 #include "ToolMenus.h"
+#include "EditorStyleSet.h"
+#include "Widgets/Input/SComboButton.h"
+#include "Widgets/Images/SImage.h"
 
 #define LOCTEXT_NAMESPACE "SourceControlCommands"
 
diff --git a/Engine/Source/Editor/StatusBar/Private/SourceControlMenuHelpers.h b/Engine/Source/Editor/StatusBar/Private/SourceControlMenuHelpers.h
--- a/Engine/Source/Editor/StatusBar/Private/SourceControlMenuHelpers.h
+++ b/Engine/Source/Editor/StatusBar/Private/SourceControlMenuHelpers.h
@@ -2,6 +2,12 @@
 
 #pragma once
 
+#include "Framework/Commands/Commands.h"
+#include "ISourceControlOperation.h"
+#include "ISourceControlProvider.h"
+
+class FUICommandList;
+class SWidget;
 
 class FSourceControlCommands : public TCommands<FSourceControlCommands>
 {
diff --git a/Engine/Source/Editor/StatusBar/Private/StatusBarSubsystem.cpp b/Engine/Source/Editor/StatusBar/Private/StatusBarSubsystem.cpp
--- a/Engine/Source/Editor/StatusBar/Private/StatusBarSubsystem.cpp
+++ b/Engine/Source/Editor/StatusBar/Private/StatusBarSubsystem.cpp
@@ -4,13 +4,14 @@
 
 #include "StatusBarSubsystem.h"
 #include "ToolMenus.h"
-
 #include "Framework/Docking/TabManager.h"
-
 #include "Framework/Application/SlateApplication.h"
 #include "ContentBrowserModule.h"
 #include "IContentBrowserSingleton.h"
-
+#include "SourceControlMenuHelpers.h"
+#include "SStatusBar.h"
+#include "Widgets/Docking/SDockTab.h"
+#include "Editor.h"
 
 
 #define LOCTEXT_NAMESPACE "StatusBar"
diff --git a/Engine/Source/Editor/StatusBar/Public/SStatusBar.h b/Engine/Source/Editor/StatusBar/Public/SStatusBar.h
--- a/Engine/Source/Editor/StatusBar/Public/SStatusBar.h
+++ b/Engine/Source/Editor/StatusBar/Public/SStatusBar.h
@@ -3,6 +3,10 @@
 #pragma once
 
 #include "StatusBarSubsystem.h"
+#include "Widgets/SCompoundWidget.h"
+#include "Widgets/DeclarativeSyntaxSupport.h"
+#include "Framework/SlateDelegates.h"
+#include "Animation/CurveSequence.h"
 
 class SStatusBar;
 
diff --git a/Engine/Source/Editor/StatusBar/Public/StatusBarSubsystem.h b/Engine/Source/Editor/StatusBar/Public/StatusBarSubsystem.h
--- a/Engine/Source/Editor/StatusBar/Public/StatusBarSubsystem.h
+++ b/Engine/Source/Editor/StatusBar/Public/StatusBarSubsystem.h
@@ -9,7 +9,9 @@
 
 
 class SStatusBar;
-
+class SWindow;
+class SWidget;
+class SDockTab;
 
 template<typename ObjectType> 
 class TAttribute;
