Added missing drop shadow to status bar content browser drawer

diff --git a/Engine/Source/Editor/StatusBar/Private/SStatusBar.cpp b/Engine/Source/Editor/StatusBar/Private/SStatusBar.cpp
--- a/Engine/Source/Editor/StatusBar/Private/SStatusBar.cpp
+++ b/Engine/Source/Editor/StatusBar/Private/SStatusBar.cpp
@@ -37,7 +37,7 @@ class SContentBrowserOverlay : public SCompoundWidget
 		SLATE_END_ARGS()
 
 
-		void SetHeight(float InHeight)
+	void SetHeight(float InHeight)
 	{
 		Height = InHeight;
 	}
@@ -53,9 +53,9 @@ class SContentBrowserOverlay : public SCompoundWidget
 		ShadowBrush = FAppStyle::Get().GetBrush("StatusBar.ContentBrowserShadow");
 
 		ChildSlot
-			[
-				InArgs._Content.Widget
-			];
+		[
+			InArgs._Content.Widget
+		];
 	}
 
 	virtual bool SupportsKeyboardFocus() const
@@ -107,6 +107,15 @@ class SContentBrowserOverlay : public SCompoundWidget
 
 		int32 OutLayerId = SCompoundWidget::OnPaint(Args, RenderTransformedChildGeometry, MyCullingRect, OutDrawElements, LayerId, InWidgetStyle, bParentEnabled);
 
+		FSlateDrawElement::MakeBox(
+			OutDrawElements,
+			OutLayerId,
+			AllottedGeometry.ToPaintGeometry(FVector2D(0.0f, AllottedGeometry.GetLocalSize().Y - ShadowOffset.Y), FVector2D(AllottedGeometry.GetLocalSize().X, ShadowOffset.Y)),
+			ShadowBrush,
+			ESlateDrawEffect::None,
+			ShadowBrush->GetTint(InWidgetStyle));
+
+
 		return OutLayerId;
 
 	}
