The content browser drawer in the status bar is now resizable

diff --git a/Engine/Source/Editor/StatusBar/Private/SStatusBar.cpp b/Engine/Source/Editor/StatusBar/Private/SStatusBar.cpp
--- a/Engine/Source/Editor/StatusBar/Private/SStatusBar.cpp
+++ b/Engine/Source/Editor/StatusBar/Private/SStatusBar.cpp
@@ -21,6 +21,10 @@
 #include "OutputLogModule.h"
 #include "SourceControlMenuHelpers.h"
 #include "Framework/Application/SlateApplication.h"
+#include "Toolkits/GlobalEditorCommonCommands.h"
+#include "Widgets/Layout/SSplitter.h"
+#include "InputCoreTypes.h"
+#include "Misc/ConfigCacheIni.h"
 
 #define LOCTEXT_NAMESPACE "StatusBar"
 
@@ -32,40 +36,55 @@ class SContentBrowserOverlay : public SCompoundWidget
 		_ShadowOffset = FVector2D(10.0f, 20.0f);
 	}
 	SLATE_DEFAULT_SLOT(FArguments, Content)
+		SLATE_ARGUMENT(float, MinContentBrowserHeight)
 		SLATE_ARGUMENT(float, MaxContentBrowserHeight)
+		SLATE_ARGUMENT(float, TargetContentBrowserHeight)
+		SLATE_EVENT(FOnContentBrowserTargetHeightChanged, OnTargetHeightChanged)
 		SLATE_ARGUMENT(FVector2D, ShadowOffset)
-		SLATE_END_ARGS()
-
+	SLATE_END_ARGS()
 
 	void SetHeight(float InHeight)
 	{
-		Height = InHeight;
+		CurrentHeight = FMath::Clamp(InHeight, MinContentBrowserHeight, TargetContentBrowserHeight);
 	}
 
 	void Construct(const FArguments& InArgs)
 	{
-		Height = 0;
+		CurrentHeight = 0;
 
 		ShadowOffset = InArgs._ShadowOffset;
+		ExpanderSize = 5.0f;
+
+		SplitterStyle = &FAppStyle::Get().GetWidgetStyle<FSplitterStyle>("Splitter");
+
+		MinContentBrowserHeight = InArgs._MinContentBrowserHeight;
+
 		MaxContentBrowserHeight = InArgs._MaxContentBrowserHeight;
 
+		TargetContentBrowserHeight = FMath::Clamp(InArgs._TargetContentBrowserHeight, MinContentBrowserHeight, MaxContentBrowserHeight);
+
+		OnTargetHeightChanged = InArgs._OnTargetHeightChanged;
+
 		BackgroundBrush = FAppStyle::Get().GetBrush("StatusBar.Background");
 		ShadowBrush = FAppStyle::Get().GetBrush("StatusBar.ContentBrowserShadow");
 
+		bIsResizeHandleHovered = false;
+		bIsResizing = false;
+
 		ChildSlot
 		[
 			InArgs._Content.Widget
 		];
 	}
 
-	virtual bool SupportsKeyboardFocus() const
+	virtual bool SupportsKeyboardFocus() const override
 	{
 		return true;
 	}
 
 	virtual FVector2D ComputeDesiredSize(float) const
 	{
-		return FVector2D(1.0f, MaxContentBrowserHeight + ShadowOffset.Y);
+		return FVector2D(1.0f, TargetContentBrowserHeight + ShadowOffset.Y);
 	}
 
 	virtual void OnArrangeChildren(const FGeometry& AllottedGeometry, FArrangedChildren& ArrangedChildren) const override
@@ -77,18 +96,100 @@ class SContentBrowserOverlay : public SCompoundWidget
 				AllottedGeometry.MakeChild(
 					ChildSlot.GetWidget(),
 					ShadowOffset,
-					FVector2D(AllottedGeometry.GetLocalSize().X - (ShadowOffset.X * 2), MaxContentBrowserHeight)
+					FVector2D(AllottedGeometry.GetLocalSize().X - (ShadowOffset.X * 2), CurrentHeight)
 				)
 			);
 		}
 	}
 
+	virtual FReply OnMouseButtonDown(const FGeometry& AllottedGeometry, const FPointerEvent& MouseEvent) override
+	{
+		FReply Reply = FReply::Unhandled();
+		if (MouseEvent.GetEffectingButton() == EKeys::LeftMouseButton)
+		{
+			const FGeometry RenderTransformedChildGeometry = GetRenderTransformedGeometry(AllottedGeometry);
+			const FGeometry ResizeHandleGeometry = GetResizeHandleGeometry(AllottedGeometry);
+
+			if (ResizeHandleGeometry.IsUnderLocation(MouseEvent.GetScreenSpacePosition()))
+			{
+				bIsResizing = true;
+				Reply = FReply::Handled().CaptureMouse(SharedThis(this));
+			}
+		}
+
+		return Reply;
+
+	}
+
+	FReply OnMouseButtonUp(const FGeometry& AllottedGeometry, const FPointerEvent& MouseEvent) override
+	{
+		if (MouseEvent.GetEffectingButton() == EKeys::LeftMouseButton && bIsResizing == true)
+		{
+			bIsResizing = false;
+			OnTargetHeightChanged.ExecuteIfBound(TargetContentBrowserHeight);
+			return FReply::Handled().ReleaseMouseCapture();
+		}
+		return FReply::Unhandled();
+	}
+
+	FReply OnMouseMove(const FGeometry& AllottedGeometry, const FPointerEvent& MouseEvent) override
+	{
+		const FGeometry RenderTransformedChildGeometry = GetRenderTransformedGeometry(AllottedGeometry);
+		const FGeometry ResizeHandleGeometry = GetResizeHandleGeometry(AllottedGeometry);
+
+		bIsResizeHandleHovered = ResizeHandleGeometry.IsUnderLocation(MouseEvent.GetScreenSpacePosition());
+
+		if (bIsResizing && this->HasMouseCapture())
+		{
+			const FVector2D LocalMousePos = ResizeHandleGeometry.AbsoluteToLocal(MouseEvent.GetScreenSpacePosition());
+
+			const float Delta = (ResizeHandleGeometry.GetLocalPositionAtCoordinates(FVector2D::ZeroVector) - LocalMousePos).Y;
+
+			TargetContentBrowserHeight = FMath::Clamp(TargetContentBrowserHeight + Delta, MinContentBrowserHeight, MaxContentBrowserHeight);
+
+			SetHeight(CurrentHeight + Delta);
+
+			return FReply::Handled();
+		}
+		else
+		{
+			return FReply::Unhandled();
+		}
+	}
+
+	virtual void OnMouseLeave(const FPointerEvent& MouseEvent) override
+	{
+		SCompoundWidget::OnMouseLeave(MouseEvent);
+
+		bIsResizeHandleHovered = false;
+	}
+
+	FCursorReply OnCursorQuery(const FGeometry& MyGeometry, const FPointerEvent& CursorEvent) const override
+	{
+		return bIsResizing || bIsResizeHandleHovered ? FCursorReply::Cursor(EMouseCursor::ResizeUpDown) : FCursorReply::Unhandled();
+	}
+
 	virtual int32 OnPaint(const FPaintArgs& Args, const FGeometry& AllottedGeometry, const FSlateRect& MyCullingRect, FSlateWindowElementList& OutDrawElements, int32 LayerId, const FWidgetStyle& InWidgetStyle, bool bParentEnabled) const override
 	{
 		static FSlateColor ShadowColor = FAppStyle::Get().GetSlateColor("Colors.Foldout");
 
-		FGeometry RenderTransformedChildGeometry = AllottedGeometry.MakeChild(FSlateRenderTransform(FVector2D(0.0f, MaxContentBrowserHeight - Height)));
+		const FGeometry RenderTransformedChildGeometry = GetRenderTransformedGeometry(AllottedGeometry);
+		const FGeometry ResizeHandleGeometry = GetResizeHandleGeometry(AllottedGeometry);
 
+		// Draw the resize handle
+		if (bIsResizing || bIsResizeHandleHovered)
+		{
+			const FSlateBrush* SplitterBrush = &SplitterStyle->HandleHighlightBrush;
+			FSlateDrawElement::MakeBox(
+				OutDrawElements,
+				LayerId,
+				ResizeHandleGeometry.ToPaintGeometry(),
+				SplitterBrush,
+				ESlateDrawEffect::None,
+				SplitterBrush->GetTint(InWidgetStyle));
+		}
+
+		// Top Shadow
 		FSlateDrawElement::MakeBox(
 			OutDrawElements,
 			LayerId,
@@ -97,16 +198,18 @@ class SContentBrowserOverlay : public SCompoundWidget
 			ESlateDrawEffect::None,
 			ShadowBrush->GetTint(InWidgetStyle));
 
+		// Background
 		FSlateDrawElement::MakeBox(
 			OutDrawElements,
 			LayerId,
-			RenderTransformedChildGeometry.ToPaintGeometry(ShadowOffset, FVector2D(AllottedGeometry.GetLocalSize().X - (ShadowOffset.X * 2), MaxContentBrowserHeight)),
+			RenderTransformedChildGeometry.ToPaintGeometry(ShadowOffset, FVector2D(AllottedGeometry.GetLocalSize().X - (ShadowOffset.X * 2), TargetContentBrowserHeight)),
 			BackgroundBrush,
 			ESlateDrawEffect::None,
 			BackgroundBrush->GetTint(InWidgetStyle));
 
 		int32 OutLayerId = SCompoundWidget::OnPaint(Args, RenderTransformedChildGeometry, MyCullingRect, OutDrawElements, LayerId, InWidgetStyle, bParentEnabled);
 
+		// Bottom shadow
 		FSlateDrawElement::MakeBox(
 			OutDrawElements,
 			OutLayerId,
@@ -115,17 +218,33 @@ class SContentBrowserOverlay : public SCompoundWidget
 			ESlateDrawEffect::None,
 			ShadowBrush->GetTint(InWidgetStyle));
 
-
 		return OutLayerId;
 
 	}
+
 private:
+	FGeometry GetRenderTransformedGeometry(const FGeometry& AllottedGeometry) const
+	{
+		return AllottedGeometry.MakeChild(FSlateRenderTransform(FVector2D(0.0f, TargetContentBrowserHeight - CurrentHeight)));
+	}
+
+	FGeometry GetResizeHandleGeometry(const FGeometry& AllottedGeometry) const
+	{
+		return GetRenderTransformedGeometry(AllottedGeometry).MakeChild(ShadowOffset - FVector2D(0.0f, ExpanderSize), FVector2D(AllottedGeometry.GetLocalSize().X, 5.0f));
+	}
+private:
+	FOnContentBrowserTargetHeightChanged OnTargetHeightChanged;
 	const FSlateBrush* BackgroundBrush;
 	const FSlateBrush* ShadowBrush;
+	const FSplitterStyle* SplitterStyle;
 	FVector2D ShadowOffset;
-	float Height;
+	float ExpanderSize;
+	float CurrentHeight;
+	float MinContentBrowserHeight;
 	float MaxContentBrowserHeight;
-
+	float TargetContentBrowserHeight;
+	bool bIsResizing;
+	bool bIsResizeHandleHovered;
 };
 
 
@@ -153,8 +272,8 @@ void SStatusBar::Construct(const FArguments& InArgs, FName InStatusBarName, cons
 	[
 		SNew(SBox)
 		.HeightOverride(FAppStyle::Get().GetFloat("StatusBar.Height"))
-		[				SNew(SHorizontalBox)
-
+		[				
+			SNew(SHorizontalBox)
 			+ SHorizontalBox::Slot()
 			.Padding(1.0f, 0.0f)
 			.AutoWidth()
@@ -321,16 +440,26 @@ void SStatusBar::OpenContentBrowser()
 
 		TSharedPtr<SWindow> MyWindow = FSlateApplication::Get().FindWidgetWindow(AsShared());
 
-		MaxContentBrowserHeight = MyWindow->GetSizeInScreen().Y * 0.33f+20;
-		MinContentBrowserHeight = GetTickSpaceGeometry().GetLocalSize().Y + MyWindow->GetWindowBorderSize().Bottom;
+		const float MaxContentBrowserHeight = MyWindow->GetSizeInScreen().Y * 0.90f;
+
+		float TargetContentBrowserHeightPct = .33f;
+		GConfig->GetFloat(TEXT("ContentBrowserDrawerSizes"), *StatusBarName.ToString(), TargetContentBrowserHeightPct, GEditorSettingsIni);
+
+		TargetContentBrowserHeight = (MyWindow->GetSizeInScreen().Y * TargetContentBrowserHeightPct) / MyWindow->GetDPIScaleFactor();
+
+
+		const float MinContentBrowserHeight = GetTickSpaceGeometry().GetLocalSize().Y + MyWindow->GetWindowBorderSize().Bottom;
 		ContentBrowserEasingCurve.Play(ThisStatusBar, false, ContentBrowserEasingCurve.IsPlaying() ? ContentBrowserEasingCurve.GetSequenceTime() : 0.0f, false);
 
 		MyWindow->AddOverlaySlot()
 			.VAlign(VAlign_Bottom)
-			.Padding(FMargin(10.0f,20.0f,10.0f, MinContentBrowserHeight))
+			.Padding(FMargin(10.0f, 20.0f, 10.0f, MinContentBrowserHeight))
 			[
 				SAssignNew(ContentBrowserOverlayContent, SContentBrowserOverlay)
+				.MinContentBrowserHeight(MinContentBrowserHeight)
+				.TargetContentBrowserHeight(TargetContentBrowserHeight)
 				.MaxContentBrowserHeight(MaxContentBrowserHeight)
+				.OnTargetHeightChanged(this, &SStatusBar::OnContentBrowserTargetHeightChanged)
 				[
 					GetContentBrowserDelegate.Execute()
 				]
@@ -417,6 +546,7 @@ TSharedRef<SWidget> SStatusBar::MakeContentBrowserWidget()
 		SNew(SButton)
 		.ButtonStyle(&FAppStyle::Get().GetWidgetStyle<FButtonStyle>("StatusBar.StatusBarButton"))
 		.OnClicked(this, &SStatusBar::OnContentBrowserButtonClicked)
+		.ToolTipText(FText::Format(LOCTEXT("StatusBar_ContentBrowserDrawerToolTip","Opens a temporary content browser above this status which will dismiss when it loses focus ({0})"), FGlobalEditorCommonCommands::Get().OpenContentBrowserDrawer->GetInputText()))
 		[
 			SNew(SHorizontalBox)
 			+ SHorizontalBox::Slot()
@@ -515,7 +645,7 @@ EActiveTimerReturnType SStatusBar::UpdateContentBrowserAnimation(double CurrentT
 {
 	check(ContentBrowserOverlayContent.IsValid());
 
-	ContentBrowserOverlayContent->SetHeight(FMath::Lerp(0.0f, MaxContentBrowserHeight, ContentBrowserEasingCurve.GetLerp()));
+	ContentBrowserOverlayContent->SetHeight(FMath::Lerp(0.0f, TargetContentBrowserHeight, ContentBrowserEasingCurve.GetLerp()));
 
 	if (!ContentBrowserEasingCurve.IsPlaying())
 	{
@@ -531,6 +661,19 @@ EActiveTimerReturnType SStatusBar::UpdateContentBrowserAnimation(double CurrentT
 	return EActiveTimerReturnType::Continue;
 }
 
+void SStatusBar::OnContentBrowserTargetHeightChanged(float TargetHeight)
+{
+	TargetContentBrowserHeight = TargetHeight;
+
+	TSharedPtr<SWindow> MyWindow = FSlateApplication::Get().FindWidgetWindow(AsShared());
+
+	// Save the height has a percentage of the screen
+	const float TargetContentBrowserHeightPct = TargetHeight / (MyWindow->GetSizeInScreen().Y / MyWindow->GetDPIScaleFactor());
+
+	// Write to config
+	GConfig->SetFloat(TEXT("ContentBrowserDrawerSizes"), *StatusBarName.ToString(), TargetContentBrowserHeightPct, GEditorSettingsIni);
+}
+
 void SStatusBar::RegisterStatusBarMenu()
 {
 	static const FName StatusBarToolBarName("StatusBar.ToolBar");
diff --git a/Engine/Source/Editor/StatusBar/Private/StatusBarSubsystem.cpp b/Engine/Source/Editor/StatusBar/Private/StatusBarSubsystem.cpp
--- a/Engine/Source/Editor/StatusBar/Private/StatusBarSubsystem.cpp
+++ b/Engine/Source/Editor/StatusBar/Private/StatusBarSubsystem.cpp
@@ -54,14 +54,33 @@ bool UStatusBarSubsystem::FocusDebugConsole(TSharedRef<SWindow> ParentWindow)
 	return bFocusedSuccessfully;
 }
 
-bool UStatusBarSubsystem::OpenContentBrowser(TSharedRef<SWindow> ParentWindow)
+bool UStatusBarSubsystem::ToggleContentBrowser(TSharedRef<SWindow> ParentWindow)
 {
-	bool bFocusedSuccessfully = false;
+	bool bWasDismissed = false;
 
-	CreateContentBrowserIfNeeded();
+	for (auto StatusBar : StatusBars)
+	{
+		if (TSharedPtr<SStatusBar> StatusBarPinned = StatusBar.Value.Pin())
+		{
+			if(StatusBarPinned->IsContentBrowserOpened())
+			{
+				TSharedPtr<SDockTab> ParentTab = StatusBarPinned->GetParentTab();
+				if (ParentTab && ParentTab->IsForeground() && ParentTab->GetParentWindow() == ParentWindow)
+				{
+					StatusBarPinned->DismissContentBrowser(nullptr);
+					bWasDismissed = true;
+				}
+			}
+		}
+	}
 
-	TSharedPtr<SWindow> Window = ParentWindow;
-	GEditor->GetTimerManager()->SetTimerForNextTick(FTimerDelegate::CreateUObject(this, &UStatusBarSubsystem::HandleDeferredOpenContentBrowser, Window));
+	if(!bWasDismissed)
+	{
+		CreateContentBrowserIfNeeded();
+
+		TSharedPtr<SWindow> Window = ParentWindow;
+		GEditor->GetTimerManager()->SetTimerForNextTick(FTimerDelegate::CreateUObject(this, &UStatusBarSubsystem::HandleDeferredOpenContentBrowser, Window));
+	}
 
 	return true;
 }
diff --git a/Engine/Source/Editor/StatusBar/Public/SStatusBar.h b/Engine/Source/Editor/StatusBar/Public/SStatusBar.h
--- a/Engine/Source/Editor/StatusBar/Public/SStatusBar.h
+++ b/Engine/Source/Editor/StatusBar/Public/SStatusBar.h
@@ -12,6 +12,8 @@ class SStatusBar;
 
 DECLARE_DELEGATE_OneParam(FOnContentBrowserOpened, TSharedRef<SStatusBar>&)
 DECLARE_DELEGATE_OneParam(FOnContentBrowserDismissed, const TSharedPtr<SWidget>&)
+DECLARE_DELEGATE_OneParam(FOnContentBrowserTargetHeightChanged, float);
+
 
 class SWindow;
 class SDockTab;
@@ -34,25 +36,36 @@ class SStatusBar : public SCompoundWidget
 	SLATE_END_ARGS()
 
 public:
+	/** SWidget interface */
 	virtual bool SupportsKeyboardFocus() const { return false; }
-
 	void Construct(const FArguments& InArgs, FName InStatusBarName, const TSharedRef<SDockTab> InParentTab);
 
+	/** 
+	 * Pushes a new status bar message
+	 *
+	 * @param InHandle	A handle to the message for clearing it later
+	 * @param InMessage	The message to display
+	 * @param InHintText	Optional hint text message.  This message will be highlighted to make it stand out
+	 */
 	void PushMessage(FStatusBarMessageHandle& InHandle, const TAttribute<FText>& InMessage, const TAttribute<FText>& InHintText);
+
+	/**
+	 * Removes a message from the status bar.  When messages are removed the previous message on the stack (if any) is displayed
+	 *
+	 * @param InHandle	Handle to the status bar message to remove
+	 */
 	void PopMessage(FStatusBarMessageHandle& InHandle);
+
+	/**
+	 * Removes all messages from the status bar
+	 */
 	void ClearAllMessages();
 
-	EVisibility GetHelpIconVisibility() const;
 
 	TSharedPtr<SDockTab> GetParentTab() const;
 	void FocusDebugConsole();
 	bool IsDebugConsoleFocused() const;
 
-
-	void OnGlobalFocusChanging(const FFocusEvent& FocusEvent, const FWeakWidgetPath& OldFocusedWidgetPath, const TSharedPtr<SWidget>& OldFocusedWidget, const FWidgetPath& NewFocusedWidgetPath, const TSharedPtr<SWidget>& NewFocusedWidget);
-
-	void OnActiveTabChanged(TSharedPtr<SDockTab> PreviouslyActive, TSharedPtr<SDockTab> NewlyActivated);
-
 	bool IsContentBrowserOpened() const;
 
 	void OpenContentBrowser();
@@ -66,9 +79,15 @@ public:
 	 * Removes a content browser from the layout immediately
 	 */
 	void RemoveContentBrowser();
+private:
+	/** Called when global focus changes which is used to determine if we should close an opened content browser drawer */
+	void OnGlobalFocusChanging(const FFocusEvent& FocusEvent, const FWeakWidgetPath& OldFocusedWidgetPath, const TSharedPtr<SWidget>& OldFocusedWidget, const FWidgetPath& NewFocusedWidgetPath, const TSharedPtr<SWidget>& NewFocusedWidget);
 
+	/** Called when active tab changes which is used to determine if we should close an opened content browser drawer */
+	void OnActiveTabChanged(TSharedPtr<SDockTab> PreviouslyActive, TSharedPtr<SDockTab> NewlyActivated);
+
+	EVisibility GetHelpIconVisibility() const;
 
-private:
 	const FSlateBrush* GetContentBrowserExpandArrowImage() const;
 
 	FText GetStatusBarMessage() const;
@@ -82,6 +101,7 @@ private:
 	FReply OnContentBrowserButtonClicked();
 
 	EActiveTimerReturnType UpdateContentBrowserAnimation(double CurrentTime, float DeltaTime);
+	void OnContentBrowserTargetHeightChanged(float TargetHeight);
 
 	void RegisterStatusBarMenu();
 
@@ -109,8 +129,7 @@ private:
 	FOnContentBrowserOpened OnContentBrowserOpenedDelegate;
 	FOnContentBrowserDismissed OnContentBrowserDismissedDelegate;
 	FCurveSequence ContentBrowserEasingCurve;
-	float MaxContentBrowserHeight;
-	float MinContentBrowserHeight;
+	float TargetContentBrowserHeight;
 	const FSlateBrush* UpArrow;
 	const FSlateBrush* DownArrow;
 	FName StatusBarName;
diff --git a/Engine/Source/Editor/StatusBar/Public/StatusBarSubsystem.h b/Engine/Source/Editor/StatusBar/Public/StatusBarSubsystem.h
--- a/Engine/Source/Editor/StatusBar/Public/StatusBarSubsystem.h
+++ b/Engine/Source/Editor/StatusBar/Public/StatusBarSubsystem.h
@@ -72,9 +72,9 @@ public:
 	bool FocusDebugConsole(TSharedRef<SWindow> ParentWindow);
 
 	/**
-	 * Opens the content browser drawer for a status bar residing in the passed in parent window
+	 * Opens or closes the content browser drawer for a status bar residing in the passed in parent window
 	 */
-	bool OpenContentBrowser(TSharedRef<SWindow> ParentWindow);
+	bool ToggleContentBrowser(TSharedRef<SWindow> ParentWindow);
 
 	/** 
 	 * Creates a new instance of a status bar widget
@@ -103,7 +103,7 @@ public:
 	void PopStatusBarMessage(FName StatusBarName, FStatusBarMessageHandle InHandle);
 
 	/**
-	 * Removes all message from the status bar
+	 * Removes all messages from the status bar
 	 *
 	 * @param StatusBarName	The name of the status bar to remove from
 	 */
diff --git a/Engine/Source/Editor/StatusBar/StatusBar.Build.cs b/Engine/Source/Editor/StatusBar/StatusBar.Build.cs
--- a/Engine/Source/Editor/StatusBar/StatusBar.Build.cs
+++ b/Engine/Source/Editor/StatusBar/StatusBar.Build.cs
@@ -16,6 +16,7 @@ public class StatusBar : ModuleRules
                 "Engine",           
 				"Slate",
                 "SlateCore",
+				"InputCore",
                 "EditorStyle",
                 "UnrealEd",
 				"ToolMenus",
diff --git a/Engine/Source/Editor/UnrealEd/Private/Toolkits/GlobalEditorCommonCommands.cpp b/Engine/Source/Editor/UnrealEd/Private/Toolkits/GlobalEditorCommonCommands.cpp
--- a/Engine/Source/Editor/UnrealEd/Private/Toolkits/GlobalEditorCommonCommands.cpp
+++ b/Engine/Source/Editor/UnrealEd/Private/Toolkits/GlobalEditorCommonCommands.cpp
@@ -171,7 +171,7 @@ void FGlobalEditorCommonCommands::OnOpenContentBrowserDrawer()
 		TSharedRef<SWindow> WindowRef = ParentWindow.ToSharedRef();
 		FOutputLogModule& OutputLogModule = FModuleManager::LoadModuleChecked<FOutputLogModule>(TEXT("OutputLog"));
 
-		GEditor->GetEditorSubsystem<UStatusBarSubsystem>()->OpenContentBrowser(ParentWindow.ToSharedRef());
+		GEditor->GetEditorSubsystem<UStatusBarSubsystem>()->ToggleContentBrowser(ParentWindow.ToSharedRef());
 	}
 }
 
