Throttle viewport rendering while moving up the content browser drawer because as the drawer grows thumbnails redraw which can cause hitching while moving

diff --git a/Engine/Source/Editor/StatusBar/Private/SStatusBar.cpp b/Engine/Source/Editor/StatusBar/Private/SStatusBar.cpp
--- a/Engine/Source/Editor/StatusBar/Private/SStatusBar.cpp
+++ b/Engine/Source/Editor/StatusBar/Private/SStatusBar.cpp
@@ -113,6 +113,8 @@ class SContentBrowserOverlay : public SCompoundWidget
 			if (ResizeHandleGeometry.IsUnderLocation(MouseEvent.GetScreenSpacePosition()))
 			{
 				bIsResizing = true;
+				ResizeThrottleHandle = FSlateThrottleManager::Get().EnterResponsiveMode();
+
 				Reply = FReply::Handled().CaptureMouse(SharedThis(this));
 			}
 		}
@@ -126,6 +128,8 @@ class SContentBrowserOverlay : public SCompoundWidget
 		if (MouseEvent.GetEffectingButton() == EKeys::LeftMouseButton && bIsResizing == true)
 		{
 			bIsResizing = false;
+			FSlateThrottleManager::Get().LeaveResponsiveMode(ResizeThrottleHandle);
+
 			OnTargetHeightChanged.ExecuteIfBound(TargetContentBrowserHeight);
 			return FReply::Handled().ReleaseMouseCapture();
 		}
@@ -230,7 +234,7 @@ private:
 
 	FGeometry GetResizeHandleGeometry(const FGeometry& AllottedGeometry) const
 	{
-		return GetRenderTransformedGeometry(AllottedGeometry).MakeChild(ShadowOffset - FVector2D(0.0f, ExpanderSize), FVector2D(AllottedGeometry.GetLocalSize().X, 5.0f));
+		return GetRenderTransformedGeometry(AllottedGeometry).MakeChild(ShadowOffset - FVector2D(0.0f, ExpanderSize), FVector2D(AllottedGeometry.GetLocalSize().X-ShadowOffset.X*2, ExpanderSize));
 	}
 private:
 	FOnContentBrowserTargetHeightChanged OnTargetHeightChanged;
@@ -238,6 +242,7 @@ private:
 	const FSlateBrush* ShadowBrush;
 	const FSplitterStyle* SplitterStyle;
 	FVector2D ShadowOffset;
+	FThrottleRequest ResizeThrottleHandle;
 	float ExpanderSize;
 	float CurrentHeight;
 	float MinContentBrowserHeight;
