Some fixes to address hitching in the CB when animating from a drawer

diff --git a/Engine/Source/Editor/StatusBar/Private/SStatusBar.cpp b/Engine/Source/Editor/StatusBar/Private/SStatusBar.cpp
--- a/Engine/Source/Editor/StatusBar/Private/SStatusBar.cpp
+++ b/Engine/Source/Editor/StatusBar/Private/SStatusBar.cpp
@@ -96,7 +96,7 @@ class SContentBrowserOverlay : public SCompoundWidget
 				AllottedGeometry.MakeChild(
 					ChildSlot.GetWidget(),
 					ShadowOffset,
-					FVector2D(AllottedGeometry.GetLocalSize().X - (ShadowOffset.X * 2), CurrentHeight)
+					FVector2D(AllottedGeometry.GetLocalSize().X - (ShadowOffset.X * 2), TargetContentBrowserHeight)
 				)
 			);
 		}
@@ -472,6 +472,7 @@ void SStatusBar::OpenContentBrowser()
 
 		if (!ContentBrowserOpenCloseTimer.IsValid())
 		{
+			AnimationThrottle = FSlateThrottleManager::Get().EnterResponsiveMode();
 			ContentBrowserOpenCloseTimer = RegisterActiveTimer(0.0f, FWidgetActiveTimerDelegate::CreateSP(ThisStatusBar, &SStatusBar::UpdateContentBrowserAnimation));
 		}
 
@@ -493,6 +494,7 @@ void SStatusBar::DismissContentBrowser(const TSharedPtr<SWidget>& NewlyFocusedWi
 
 		if (!ContentBrowserOpenCloseTimer.IsValid())
 		{
+			AnimationThrottle = FSlateThrottleManager::Get().EnterResponsiveMode();
 			ContentBrowserOpenCloseTimer = RegisterActiveTimer(0.0f, FWidgetActiveTimerDelegate::CreateSP(this, &SStatusBar::UpdateContentBrowserAnimation));
 		}
 
@@ -515,6 +517,7 @@ void SStatusBar::RemoveContentBrowser()
 
 		if (ContentBrowserOpenCloseTimer.IsValid())
 		{
+			FSlateThrottleManager::Get().LeaveResponsiveMode(AnimationThrottle);
 			UnRegisterActiveTimer(ContentBrowserOpenCloseTimer.ToSharedRef());
 		}
 
@@ -659,6 +662,7 @@ EActiveTimerReturnType SStatusBar::UpdateContentBrowserAnimation(double CurrentT
 			RemoveContentBrowser();
 		}
 
+		FSlateThrottleManager::Get().LeaveResponsiveMode(AnimationThrottle);
 		ContentBrowserOpenCloseTimer.Reset();
 		return EActiveTimerReturnType::Stop;
 	}
diff --git a/Engine/Source/Editor/StatusBar/Public/SStatusBar.h b/Engine/Source/Editor/StatusBar/Public/SStatusBar.h
--- a/Engine/Source/Editor/StatusBar/Public/SStatusBar.h
+++ b/Engine/Source/Editor/StatusBar/Public/SStatusBar.h
@@ -7,6 +7,7 @@
 #include "Widgets/DeclarativeSyntaxSupport.h"
 #include "Framework/SlateDelegates.h"
 #include "Animation/CurveSequence.h"
+#include "Application/ThrottleManager.h"
 
 class SStatusBar;
 
@@ -100,6 +101,8 @@ private:
 
 	FReply OnContentBrowserButtonClicked();
 
+	void CreateAnimationTimerIfNeeded();
+
 	EActiveTimerReturnType UpdateContentBrowserAnimation(double CurrentTime, float DeltaTime);
 	void OnContentBrowserTargetHeightChanged(float TargetHeight);
 
@@ -129,6 +132,7 @@ private:
 	FOnContentBrowserOpened OnContentBrowserOpenedDelegate;
 	FOnContentBrowserDismissed OnContentBrowserDismissedDelegate;
 	FCurveSequence ContentBrowserEasingCurve;
+	FThrottleRequest AnimationThrottle;
 	float TargetContentBrowserHeight;
 	const FSlateBrush* UpArrow;
 	const FSlateBrush* DownArrow;
