Fix dragging content browser drawer sometimes causing erratic movement

diff --git a/Engine/Source/Editor/StatusBar/Private/SStatusBar.cpp b/Engine/Source/Editor/StatusBar/Private/SStatusBar.cpp
--- a/Engine/Source/Editor/StatusBar/Private/SStatusBar.cpp
+++ b/Engine/Source/Editor/StatusBar/Private/SStatusBar.cpp
@@ -113,6 +113,8 @@ class SContentBrowserOverlay : public SCompoundWidget
 			if (ResizeHandleGeometry.IsUnderLocation(MouseEvent.GetScreenSpacePosition()))
 			{
 				bIsResizing = true;
+				InitialResizeGeometry = ResizeHandleGeometry;
+				InitialContentBrowserHeightAtResize = CurrentHeight;
 				ResizeThrottleHandle = FSlateThrottleManager::Get().EnterResponsiveMode();
 
 				Reply = FReply::Handled().CaptureMouse(SharedThis(this));
@@ -138,20 +140,18 @@ class SContentBrowserOverlay : public SCompoundWidget
 
 	FReply OnMouseMove(const FGeometry& AllottedGeometry, const FPointerEvent& MouseEvent) override
 	{
-		const FGeometry RenderTransformedChildGeometry = GetRenderTransformedGeometry(AllottedGeometry);
 		const FGeometry ResizeHandleGeometry = GetResizeHandleGeometry(AllottedGeometry);
 
 		bIsResizeHandleHovered = ResizeHandleGeometry.IsUnderLocation(MouseEvent.GetScreenSpacePosition());
 
-		if (bIsResizing && this->HasMouseCapture())
+		if (bIsResizing && this->HasMouseCapture() && !MouseEvent.GetCursorDelta().IsZero())
 		{
-			const FVector2D LocalMousePos = ResizeHandleGeometry.AbsoluteToLocal(MouseEvent.GetScreenSpacePosition());
-
-			const float Delta = (ResizeHandleGeometry.GetLocalPositionAtCoordinates(FVector2D::ZeroVector) - LocalMousePos).Y;
+			const FVector2D LocalMousePos = InitialResizeGeometry.AbsoluteToLocal(MouseEvent.GetScreenSpacePosition());
+			const float DeltaHeight = (InitialResizeGeometry.GetLocalPositionAtCoordinates(FVector2D::ZeroVector) - LocalMousePos).Y;
 
-			TargetContentBrowserHeight = FMath::Clamp(TargetContentBrowserHeight + Delta, MinContentBrowserHeight, MaxContentBrowserHeight);
+			TargetContentBrowserHeight = FMath::Clamp(InitialContentBrowserHeightAtResize + DeltaHeight, MinContentBrowserHeight, MaxContentBrowserHeight);
+			SetHeight(InitialContentBrowserHeightAtResize + DeltaHeight);
 
-			SetHeight(CurrentHeight + Delta);
 
 			return FReply::Handled();
 		}
@@ -237,6 +237,7 @@ private:
 		return GetRenderTransformedGeometry(AllottedGeometry).MakeChild(ShadowOffset - FVector2D(0.0f, ExpanderSize), FVector2D(AllottedGeometry.GetLocalSize().X-ShadowOffset.X*2, ExpanderSize));
 	}
 private:
+	FGeometry InitialResizeGeometry;
 	FOnContentBrowserTargetHeightChanged OnTargetHeightChanged;
 	const FSlateBrush* BackgroundBrush;
 	const FSlateBrush* ShadowBrush;
@@ -248,6 +249,7 @@ private:
 	float MinContentBrowserHeight;
 	float MaxContentBrowserHeight;
 	float TargetContentBrowserHeight;
+	float InitialContentBrowserHeightAtResize;
 	bool bIsResizing;
 	bool bIsResizeHandleHovered;
 };
