Content browser drawer now saves the content browser settings properly

diff --git a/Engine/Source/Editor/ContentBrowser/Private/ContentBrowserSingleton.cpp b/Engine/Source/Editor/ContentBrowser/Private/ContentBrowserSingleton.cpp
--- a/Engine/Source/Editor/ContentBrowser/Private/ContentBrowserSingleton.cpp
+++ b/Engine/Source/Editor/ContentBrowser/Private/ContentBrowserSingleton.cpp
@@ -226,14 +226,21 @@ void FContentBrowserSingleton::FocusPrimaryContentBrowser(bool bFocusSearch)
 	}
 }
 
-void FContentBrowserSingleton::FocusContentBrowserSearchField(TSharedRef<SWidget> ContentBrowserWidget)
+void FContentBrowserSingleton::FocusContentBrowserSearchField(FName InstanceName)
 {
-	static FName ContentBrowserType(TEXT("SContentBrowser"));
-
-	if (ContentBrowserWidget->GetType() == ContentBrowserType)
+	if (InstanceName.IsValid() && !InstanceName.IsNone())
 	{
-		TSharedRef<SContentBrowser> ContentBrowser = StaticCastSharedRef<SContentBrowser>(ContentBrowserWidget);
-		ContentBrowser->SetKeyboardFocusOnSearch();
+		for (TWeakPtr<SContentBrowser>& Browser : AllContentBrowsers)
+		{
+			if (TSharedPtr<SContentBrowser> BrowserPinned = Browser.Pin())
+			{
+				if (BrowserPinned->GetInstanceName() == InstanceName)
+				{
+					BrowserPinned->SetKeyboardFocusOnSearch();
+					break;
+				}
+			}
+		}
 	}
 }
 
@@ -744,6 +751,24 @@ void FContentBrowserSingleton::ForceShowPluginContent(bool bEnginePlugin)
 	}
 }
 
+void FContentBrowserSingleton::SaveContentBrowserSettings(FName InstanceName)
+{
+	if (InstanceName.IsValid() && !InstanceName.IsNone())
+	{
+		for (TWeakPtr<SContentBrowser>& Browser : AllContentBrowsers)
+		{
+			if (TSharedPtr<SContentBrowser> BrowserPinned = Browser.Pin())
+			{
+				if (BrowserPinned->GetInstanceName() == InstanceName)
+				{
+					BrowserPinned->SaveSettings();
+					break;
+				}
+			}
+		}
+	}
+}
+
 void FContentBrowserSingleton::PopulateConfigValues()
 {
 	const FString ContentBrowserSection = TEXT("ContentBrowser");
diff --git a/Engine/Source/Editor/ContentBrowser/Private/ContentBrowserSingleton.h b/Engine/Source/Editor/ContentBrowser/Private/ContentBrowserSingleton.h
--- a/Engine/Source/Editor/ContentBrowser/Private/ContentBrowserSingleton.h
+++ b/Engine/Source/Editor/ContentBrowser/Private/ContentBrowserSingleton.h
@@ -57,7 +57,7 @@ public:
 	virtual FString CreateModalSaveAssetDialog(const FSaveAssetDialogConfig& SaveAssetConfig) override;
 	virtual bool HasPrimaryContentBrowser() const override;
 	virtual void FocusPrimaryContentBrowser(bool bFocusSearch) override;
-	virtual void FocusContentBrowserSearchField(TSharedRef<SWidget> ContentBrowserWidget) override;
+	virtual void FocusContentBrowserSearchField(FName InstanceName) override;
 	virtual void CreateNewAsset(const FString& DefaultAssetName, const FString& PackagePath, UClass* AssetClass, UFactory* Factory) override;
 	virtual void SyncBrowserToAssets(const TArray<struct FAssetData>& AssetDataList, bool bAllowLockedBrowsers = false, bool bFocusContentBrowser = true, const FName& InstanceName = FName(), bool bNewSpawnBrowser = false) override;
 	virtual void SyncBrowserToAssets(const TArray<UObject*>& AssetList, bool bAllowLockedBrowsers = false, bool bFocusContentBrowser = true, const FName& InstanceName = FName(), bool bNewSpawnBrowser = false) override;
@@ -70,7 +70,7 @@ public:
 	virtual void CaptureThumbnailFromViewport(FViewport* InViewport, TArray<FAssetData>& SelectedAssets) override;
 	virtual void SetSelectedPaths(const TArray<FString>& FolderPaths, bool bNeedsRefresh = false) override;
 	virtual void ForceShowPluginContent(bool bEnginePlugin) override;
-
+	virtual void SaveContentBrowserSettings(FName InstanceName) override;
 
 	/** Gets the content browser singleton as a FContentBrowserSingleton */
 	static FContentBrowserSingleton& Get();
diff --git a/Engine/Source/Editor/ContentBrowser/Public/IContentBrowserSingleton.h b/Engine/Source/Editor/ContentBrowser/Public/IContentBrowserSingleton.h
--- a/Engine/Source/Editor/ContentBrowser/Public/IContentBrowserSingleton.h
+++ b/Engine/Source/Editor/ContentBrowser/Public/IContentBrowserSingleton.h
@@ -518,7 +518,7 @@ public:
 	virtual void FocusPrimaryContentBrowser(bool bFocusSearch) = 0;
 
 	/** Focuses the search field of a content browser widget */
-	virtual void FocusContentBrowserSearchField(TSharedRef<SWidget> ContentBrowserWidget) = 0;
+	virtual void FocusContentBrowserSearchField(FName InstanceName) = 0;
 
 	/** Sets up an inline-name for the creation of a new asset in the primary content browser using the specified path and the specified class and/or factory */
 	virtual void CreateNewAsset(const FString& DefaultAssetName, const FString& PackagePath, UClass* AssetClass, UFactory* Factory) = 0;
@@ -591,9 +591,16 @@ public:
 	virtual void SetSelectedPaths(const TArray<FString>& FolderPaths, bool bNeedsRefresh = false) = 0;
 
 	/**
-	* Forces the content browser to show plugin content if it's not already showing.
-	*
-	* @param bEnginePlugin	If this is true, it will also force the content browser to show engine content
-	*/
+	 * Forces the content browser to show plugin content if it's not already showing.
+	 *
+	 * @param bEnginePlugin	If this is true, it will also force the content browser to show engine content
+	 */
 	virtual void ForceShowPluginContent(bool bEnginePlugin) = 0;
+
+	/**
+	 * Saves the settings for a particular content browser instance
+	 *
+	 * @param InstanceName The name of the content browser to save
+	 */
+	virtual void SaveContentBrowserSettings(FName InstanceName) = 0;
 };
diff --git a/Engine/Source/Editor/StatusBar/Private/StatusBarSubsystem.cpp b/Engine/Source/Editor/StatusBar/Private/StatusBarSubsystem.cpp
--- a/Engine/Source/Editor/StatusBar/Private/StatusBarSubsystem.cpp
+++ b/Engine/Source/Editor/StatusBar/Private/StatusBarSubsystem.cpp
@@ -19,6 +19,7 @@
 
 int32 UStatusBarSubsystem::HandleCounter = 0;
 
+static const FName StatusBarContentBrowserName = "StatusBarContentBrowser";
 
 void UStatusBarSubsystem::Initialize(FSubsystemCollectionBase& Collection)
 {
@@ -161,7 +162,7 @@ void UStatusBarSubsystem::CreateContentBrowserIfNeeded()
 		FContentBrowserConfig Config;
 		Config.bCanSetAsPrimaryBrowser = false;
 
-		StatusBarContentBrowser = ContentBrowserSingleton.CreateContentBrowser("StatusBarContentBrowser", nullptr, &Config);
+		StatusBarContentBrowser = ContentBrowserSingleton.CreateContentBrowser(StatusBarContentBrowserName, nullptr, &Config);
 	}
 }
 
@@ -197,9 +198,11 @@ void UStatusBarSubsystem::OnContentBrowserDismissed(const TSharedPtr<SWidget>& N
 	if (PreviousKeyboardFocusedWidget.IsValid() && !NewlyFocusedWidget.IsValid())
 	{
 		FSlateApplication::Get().SetKeyboardFocus(PreviousKeyboardFocusedWidget.Pin());
-		
 	}
 
+	IContentBrowserSingleton& ContentBrowserSingleton = FModuleManager::Get().LoadModuleChecked<FContentBrowserModule>("ContentBrowser").Get();;
+	ContentBrowserSingleton.SaveContentBrowserSettings(StatusBarContentBrowserName);
+
 	PreviousKeyboardFocusedWidget.Reset();
 }
 
@@ -218,7 +221,7 @@ void UStatusBarSubsystem::HandleDeferredOpenContentBrowser(TSharedPtr<SWindow> P
 				// Cache off the previously focused widget so we can restore focus if the user hits the focus key again
 				PreviousKeyboardFocusedWidget = FSlateApplication::Get().GetKeyboardFocusedWidget();
 
-				ContentBrowserSingleton.FocusContentBrowserSearchField(StatusBarContentBrowser.ToSharedRef());
+				ContentBrowserSingleton.FocusContentBrowserSearchField(StatusBarContentBrowserName);
 				break;
 			}
 		}
