Fix flickering in the status bar source control menu

diff --git a/Engine/Source/Editor/StatusBar/Private/SStatusBar.cpp b/Engine/Source/Editor/StatusBar/Private/SStatusBar.cpp
--- a/Engine/Source/Editor/StatusBar/Private/SStatusBar.cpp
+++ b/Engine/Source/Editor/StatusBar/Private/SStatusBar.cpp
@@ -258,6 +258,8 @@ private:
 void SStatusBar::Construct(const FArguments& InArgs, FName InStatusBarName, const TSharedRef<SDockTab> InParentTab)
 {
 	StatusBarName = InStatusBarName;
+	StatusBarToolBarName = FName(*(StatusBarName.ToString() + ".ToolBar"));
+	
 	ParentTab = InParentTab;
 
 	UpArrow = FAppStyle::Get().GetBrush("StatusBar.ContentBrowserUp");
@@ -593,11 +595,11 @@ TSharedRef<SWidget> SStatusBar::MakeContentBrowserWidget()
 TSharedRef<SWidget> SStatusBar::MakeStatusBarToolBarWidget()
 {
 	RegisterStatusBarMenu();
-
+	
 	FToolMenuContext MenuContext;
 	RegisterSourceControlStatus();
 
-	return UToolMenus::Get()->GenerateWidget("StatusBar.ToolBar", MenuContext);
+	return UToolMenus::Get()->GenerateWidget(StatusBarToolBarName, MenuContext);
 }
 
 TSharedRef<SWidget> SStatusBar::MakeDebugConsoleWidget(FSimpleDelegate OnConsoleClosed)
@@ -687,7 +689,6 @@ void SStatusBar::OnContentBrowserTargetHeightChanged(float TargetHeight)
 
 void SStatusBar::RegisterStatusBarMenu()
 {
-	static const FName StatusBarToolBarName("StatusBar.ToolBar");
 	UToolMenus* ToolMenus = UToolMenus::Get();
 	if (ToolMenus->IsMenuRegistered(StatusBarToolBarName))
 	{
@@ -703,7 +704,7 @@ void SStatusBar::RegisterSourceControlStatus()
 	// Source Control preferences
 	FSourceControlMenuHelpers::CheckSourceControlStatus();
 	{
-		UToolMenu* SourceControlMenu = UToolMenus::Get()->ExtendMenu("StatusBar.ToolBar");
+		UToolMenu* SourceControlMenu = UToolMenus::Get()->ExtendMenu(StatusBarToolBarName);
 		FToolMenuSection& Section = SourceControlMenu->FindOrAddSection("SourceControl");
 
 		Section.AddEntry(
diff --git a/Engine/Source/Editor/StatusBar/Public/SStatusBar.h b/Engine/Source/Editor/StatusBar/Public/SStatusBar.h
--- a/Engine/Source/Editor/StatusBar/Public/SStatusBar.h
+++ b/Engine/Source/Editor/StatusBar/Public/SStatusBar.h
@@ -137,5 +137,6 @@ private:
 	const FSlateBrush* UpArrow;
 	const FSlateBrush* DownArrow;
 	FName StatusBarName;
+	FName StatusBarToolBarName;
 };
 
