Ensure any open content browser drawer is properly cleaned up when a status bar is destroyed

diff --git a/Engine/Source/Editor/StatusBar/Private/SStatusBar.cpp b/Engine/Source/Editor/StatusBar/Private/SStatusBar.cpp
--- a/Engine/Source/Editor/StatusBar/Private/SStatusBar.cpp
+++ b/Engine/Source/Editor/StatusBar/Private/SStatusBar.cpp
@@ -255,6 +255,12 @@ private:
 };
 
 
+SStatusBar::~SStatusBar()
+{
+	// Ensure the content browser is removed if we're being destroyed
+	RemoveContentBrowser();
+}
+
 void SStatusBar::Construct(const FArguments& InArgs, FName InStatusBarName, const TSharedRef<SDockTab> InParentTab)
 {
 	StatusBarName = InStatusBarName;
@@ -474,6 +480,8 @@ void SStatusBar::OpenContentBrowser()
 				]
 			];
 
+		WindowWithOverlayContent = MyWindow;
+
 		if (!ContentBrowserOpenCloseTimer.IsValid())
 		{
 			AnimationThrottle = FSlateThrottleManager::Get().EnterResponsiveMode();
@@ -513,11 +521,14 @@ void SStatusBar::RemoveContentBrowser()
 {
 	if(ContentBrowserOverlayContent.IsValid())
 	{
-		TSharedPtr<SWindow> Window = FSlateApplication::Get().FindWidgetWindow(SharedThis(this));
-
-		Window->RemoveOverlaySlot(ContentBrowserOverlayContent.ToSharedRef());
+		// Remove the content browser from the window
+		if (TSharedPtr<SWindow> Window = WindowWithOverlayContent.Pin())
+		{
+			Window->RemoveOverlaySlot(ContentBrowserOverlayContent.ToSharedRef());
+		}
 
 		ContentBrowserOverlayContent.Reset();
+		WindowWithOverlayContent.Reset();
 
 		if (ContentBrowserOpenCloseTimer.IsValid())
 		{
@@ -678,7 +689,7 @@ void SStatusBar::OnContentBrowserTargetHeightChanged(float TargetHeight)
 {
 	TargetContentBrowserHeight = TargetHeight;
 
-	TSharedPtr<SWindow> MyWindow = FSlateApplication::Get().FindWidgetWindow(AsShared());
+	TSharedPtr<SWindow> MyWindow = WindowWithOverlayContent.Pin();
 
 	// Save the height has a percentage of the screen
 	const float TargetContentBrowserHeightPct = TargetHeight / (MyWindow->GetSizeInScreen().Y / MyWindow->GetDPIScaleFactor());
diff --git a/Engine/Source/Editor/StatusBar/Public/SStatusBar.h b/Engine/Source/Editor/StatusBar/Public/SStatusBar.h
--- a/Engine/Source/Editor/StatusBar/Public/SStatusBar.h
+++ b/Engine/Source/Editor/StatusBar/Public/SStatusBar.h
@@ -37,10 +37,11 @@ class SStatusBar : public SCompoundWidget
 	SLATE_END_ARGS()
 
 public:
+	~SStatusBar();
+
 	/** SWidget interface */
 	virtual bool SupportsKeyboardFocus() const { return false; }
 	void Construct(const FArguments& InArgs, FName InStatusBarName, const TSharedRef<SDockTab> InParentTab);
-
 	/** 
 	 * Pushes a new status bar message
 	 *
@@ -127,6 +128,7 @@ private:
 	TSharedPtr<SMultiLineEditableTextBox> ConsoleEditBox;
 	TWeakPtr<SDockTab> ParentTab;
 	TSharedPtr<SContentBrowserOverlay> ContentBrowserOverlayContent;
+	TWeakPtr<SWindow> WindowWithOverlayContent;
 	TSharedPtr<FActiveTimerHandle> ContentBrowserOpenCloseTimer;
 	FOnGetContent GetContentBrowserDelegate;
 	FOnContentBrowserOpened OnContentBrowserOpenedDelegate;
