Fix some instances of the content browser drawer dismissing when menus are opened inside of it

diff --git a/Engine/Source/Editor/StatusBar/Private/SStatusBar.cpp b/Engine/Source/Editor/StatusBar/Private/SStatusBar.cpp
--- a/Engine/Source/Editor/StatusBar/Private/SStatusBar.cpp
+++ b/Engine/Source/Editor/StatusBar/Private/SStatusBar.cpp
@@ -417,7 +417,7 @@ void SStatusBar::OnGlobalFocusChanging(const FFocusEvent& FocusEvent, const FWea
 				FWidgetPath MenuHostPath;
 
 				// See if the menu being opened is part of the content browser path and if so the menu should not be dismissed
-				FSlateApplication::Get().GeneratePathToWidgetUnchecked(MenuHost.ToSharedRef(), MenuHostPath);
+				FSlateApplication::Get().GeneratePathToWidgetUnchecked(MenuHost.ToSharedRef(), MenuHostPath, EVisibility::All);
 				if (!MenuHostPath.ContainsWidget(ContentBrowserOverlayContent.ToSharedRef()))
 				{
 					bShouldDismiss = true;
