Generized the content browser drawer in the status bar so it can be used for other status bar based drawers

diff --git a/Engine/Source/Editor/StatusBar/Private/SStatusBar.cpp b/Engine/Source/Editor/StatusBar/Private/SStatusBar.cpp
--- a/Engine/Source/Editor/StatusBar/Private/SStatusBar.cpp
+++ b/Engine/Source/Editor/StatusBar/Private/SStatusBar.cpp
@@ -21,31 +21,33 @@
 #include "OutputLogModule.h"
 #include "SourceControlMenuHelpers.h"
 #include "Framework/Application/SlateApplication.h"
-#include "Toolkits/GlobalEditorCommonCommands.h"
 #include "Widgets/Layout/SSplitter.h"
 #include "InputCoreTypes.h"
 #include "Misc/ConfigCacheIni.h"
 
 #define LOCTEXT_NAMESPACE "StatusBar"
 
-class SContentBrowserOverlay : public SCompoundWidget
+
+class SDrawerOverlay : public SCompoundWidget
 {
-	SLATE_BEGIN_ARGS(SContentBrowserOverlay)
+	SLATE_BEGIN_ARGS(SDrawerOverlay)
 	{
 		_Clipping = EWidgetClipping::ClipToBounds;
 		_ShadowOffset = FVector2D(10.0f, 20.0f);
 	}
 	SLATE_DEFAULT_SLOT(FArguments, Content)
-		SLATE_ARGUMENT(float, MinContentBrowserHeight)
-		SLATE_ARGUMENT(float, MaxContentBrowserHeight)
-		SLATE_ARGUMENT(float, TargetContentBrowserHeight)
-		SLATE_EVENT(FOnContentBrowserTargetHeightChanged, OnTargetHeightChanged)
+		SLATE_ARGUMENT(float, MinDrawerHeight)
+		SLATE_ARGUMENT(float, MaxDrawerHeight)
+		SLATE_ARGUMENT(float, TargetDrawerHeight)
+		SLATE_EVENT(FOnStatusBarDrawerTargetHeightChanged, OnTargetHeightChanged)
 		SLATE_ARGUMENT(FVector2D, ShadowOffset)
 	SLATE_END_ARGS()
 
-	void SetHeight(float InHeight)
+	void UpdateHeightInterp(float InAlpha)
 	{
-		CurrentHeight = FMath::Clamp(InHeight, MinContentBrowserHeight, TargetContentBrowserHeight);
+		float NewHeight = FMath::Lerp(0.0f, TargetHeight, InAlpha);
+
+		SetHeight(NewHeight);
 	}
 
 	void Construct(const FArguments& InArgs)
@@ -57,11 +59,11 @@ class SContentBrowserOverlay : public SCompoundWidget
 
 		SplitterStyle = &FAppStyle::Get().GetWidgetStyle<FSplitterStyle>("Splitter");
 
-		MinContentBrowserHeight = InArgs._MinContentBrowserHeight;
+		MinHeight = InArgs._MinDrawerHeight;
 
-		MaxContentBrowserHeight = InArgs._MaxContentBrowserHeight;
+		MaxHeight = InArgs._MaxDrawerHeight;
 
-		TargetContentBrowserHeight = FMath::Clamp(InArgs._TargetContentBrowserHeight, MinContentBrowserHeight, MaxContentBrowserHeight);
+		TargetHeight = FMath::Clamp(InArgs._TargetDrawerHeight, MinHeight, MaxHeight);
 
 		OnTargetHeightChanged = InArgs._OnTargetHeightChanged;
 
@@ -84,7 +86,7 @@ class SContentBrowserOverlay : public SCompoundWidget
 
 	virtual FVector2D ComputeDesiredSize(float) const
 	{
-		return FVector2D(1.0f, TargetContentBrowserHeight + ShadowOffset.Y);
+		return FVector2D(1.0f, TargetHeight + ShadowOffset.Y);
 	}
 
 	virtual void OnArrangeChildren(const FGeometry& AllottedGeometry, FArrangedChildren& ArrangedChildren) const override
@@ -96,7 +98,7 @@ class SContentBrowserOverlay : public SCompoundWidget
 				AllottedGeometry.MakeChild(
 					ChildSlot.GetWidget(),
 					ShadowOffset,
-					FVector2D(AllottedGeometry.GetLocalSize().X - (ShadowOffset.X * 2), TargetContentBrowserHeight)
+					FVector2D(AllottedGeometry.GetLocalSize().X - (ShadowOffset.X * 2), TargetHeight)
 				)
 			);
 		}
@@ -114,7 +116,7 @@ class SContentBrowserOverlay : public SCompoundWidget
 			{
 				bIsResizing = true;
 				InitialResizeGeometry = ResizeHandleGeometry;
-				InitialContentBrowserHeightAtResize = CurrentHeight;
+				InitialHeightAtResize = CurrentHeight;
 				ResizeThrottleHandle = FSlateThrottleManager::Get().EnterResponsiveMode();
 
 				Reply = FReply::Handled().CaptureMouse(SharedThis(this));
@@ -132,7 +134,7 @@ class SContentBrowserOverlay : public SCompoundWidget
 			bIsResizing = false;
 			FSlateThrottleManager::Get().LeaveResponsiveMode(ResizeThrottleHandle);
 
-			OnTargetHeightChanged.ExecuteIfBound(TargetContentBrowserHeight);
+			OnTargetHeightChanged.ExecuteIfBound(TargetHeight);
 			return FReply::Handled().ReleaseMouseCapture();
 		}
 		return FReply::Unhandled();
@@ -149,8 +151,8 @@ class SContentBrowserOverlay : public SCompoundWidget
 			const FVector2D LocalMousePos = InitialResizeGeometry.AbsoluteToLocal(MouseEvent.GetScreenSpacePosition());
 			const float DeltaHeight = (InitialResizeGeometry.GetLocalPositionAtCoordinates(FVector2D::ZeroVector) - LocalMousePos).Y;
 
-			TargetContentBrowserHeight = FMath::Clamp(InitialContentBrowserHeightAtResize + DeltaHeight, MinContentBrowserHeight, MaxContentBrowserHeight);
-			SetHeight(InitialContentBrowserHeightAtResize + DeltaHeight);
+			TargetHeight = FMath::Clamp(InitialHeightAtResize + DeltaHeight, MinHeight, MaxHeight);
+			SetHeight(InitialHeightAtResize + DeltaHeight);
 
 
 			return FReply::Handled();
@@ -206,7 +208,7 @@ class SContentBrowserOverlay : public SCompoundWidget
 		FSlateDrawElement::MakeBox(
 			OutDrawElements,
 			LayerId,
-			RenderTransformedChildGeometry.ToPaintGeometry(ShadowOffset, FVector2D(AllottedGeometry.GetLocalSize().X - (ShadowOffset.X * 2), TargetContentBrowserHeight)),
+			RenderTransformedChildGeometry.ToPaintGeometry(ShadowOffset, FVector2D(AllottedGeometry.GetLocalSize().X - (ShadowOffset.X * 2), TargetHeight)),
 			BackgroundBrush,
 			ESlateDrawEffect::None,
 			BackgroundBrush->GetTint(InWidgetStyle));
@@ -229,16 +231,21 @@ class SContentBrowserOverlay : public SCompoundWidget
 private:
 	FGeometry GetRenderTransformedGeometry(const FGeometry& AllottedGeometry) const
 	{
-		return AllottedGeometry.MakeChild(FSlateRenderTransform(FVector2D(0.0f, TargetContentBrowserHeight - CurrentHeight)));
+		return AllottedGeometry.MakeChild(FSlateRenderTransform(FVector2D(0.0f, TargetHeight - CurrentHeight)));
 	}
 
 	FGeometry GetResizeHandleGeometry(const FGeometry& AllottedGeometry) const
 	{
 		return GetRenderTransformedGeometry(AllottedGeometry).MakeChild(ShadowOffset - FVector2D(0.0f, ExpanderSize), FVector2D(AllottedGeometry.GetLocalSize().X-ShadowOffset.X*2, ExpanderSize));
 	}
+
+	void SetHeight(float NewHeight)
+	{
+		CurrentHeight = FMath::Clamp(NewHeight, MinHeight, TargetHeight);
+	}
 private:
 	FGeometry InitialResizeGeometry;
-	FOnContentBrowserTargetHeightChanged OnTargetHeightChanged;
+	FOnStatusBarDrawerTargetHeightChanged OnTargetHeightChanged;
 	const FSlateBrush* BackgroundBrush;
 	const FSlateBrush* ShadowBrush;
 	const FSplitterStyle* SplitterStyle;
@@ -246,10 +253,10 @@ private:
 	FThrottleRequest ResizeThrottleHandle;
 	float ExpanderSize;
 	float CurrentHeight;
-	float MinContentBrowserHeight;
-	float MaxContentBrowserHeight;
-	float TargetContentBrowserHeight;
-	float InitialContentBrowserHeightAtResize;
+	float MinHeight;
+	float MaxHeight;
+	float TargetHeight;
+	float InitialHeightAtResize;
 	bool bIsResizing;
 	bool bIsResizeHandleHovered;
 };
@@ -258,7 +265,7 @@ private:
 SStatusBar::~SStatusBar()
 {
 	// Ensure the content browser is removed if we're being destroyed
-	RemoveContentBrowser();
+	CloseDrawerImmediately();
 }
 
 void SStatusBar::Construct(const FArguments& InArgs, FName InStatusBarName, const TSharedRef<SDockTab> InParentTab)
@@ -273,71 +280,64 @@ void SStatusBar::Construct(const FArguments& InArgs, FName InStatusBarName, cons
 
 	const FSlateBrush* StatusBarBackground = FAppStyle::Get().GetBrush("StatusBar.Background");
 
-	GetContentBrowserDelegate = InArgs._OnGetContentBrowser;
-	OnContentBrowserOpenedDelegate = InArgs._OnContentBrowserOpened;
-	OnContentBrowserDismissedDelegate = InArgs._OnContentBrowserDismissed;
-
-	ContentBrowserEasingCurve = FCurveSequence(0.0f, 0.15f, ECurveEaseFunction::QuadOut);
+	DrawerEasingCurve = FCurveSequence(0.0f, 0.15f, ECurveEaseFunction::QuadOut);
 
 	FSlateApplication::Get().OnFocusChanging().AddSP(this, &SStatusBar::OnGlobalFocusChanging);
 	FGlobalTabmanager::Get()->OnActiveTabChanged_Subscribe(FOnActiveTabChanged::FDelegate::CreateSP(this, &SStatusBar::OnActiveTabChanged));
 	FGlobalTabmanager::Get()->OnTabForegrounded_Subscribe(FOnActiveTabChanged::FDelegate::CreateSP(this, &SStatusBar::OnActiveTabChanged));
-
+	
 	ChildSlot
 	[
 		SNew(SBox)
 		.HeightOverride(FAppStyle::Get().GetFloat("StatusBar.Height"))
-		[				
+		[	
 			SNew(SHorizontalBox)
-			+ SHorizontalBox::Slot()
-			.Padding(1.0f, 0.0f)
-			.AutoWidth()
-			[
-				SNew(SBorder)
-				.Padding(FMargin(2.0f, 0.0f))
-				.BorderImage(StatusBarBackground)
-				.VAlign(VAlign_Center)
-				[			
-					MakeContentBrowserWidget()
-				]
-			] 
-			+ SHorizontalBox::Slot()
+			+SHorizontalBox::Slot()
 			.AutoWidth()
-			.Padding(1.0f, 0.0f)
 			[
-				SNew(SBorder)
-				.Padding(0.0f)
-				.BorderImage(StatusBarBackground)
-				.VAlign(VAlign_Center)
-				.Padding(FMargin(6.0f, 0.0f))
-				[
-					MakeDebugConsoleWidget(InArgs._OnConsoleClosed)
-				]
+				SAssignNew(DrawerBox, SHorizontalBox)
 			]
-			+ SHorizontalBox::Slot()
-			.FillWidth(1.0f)
-			.Padding(1.0f, 0.0f)
+			+SHorizontalBox::Slot()
 			[
-				SNew(SBorder)
-				.Padding(0.0f)
-				.BorderImage(StatusBarBackground)
-				.VAlign(VAlign_Center)
-				.Padding(FMargin(6.0f, 0.0f))
+				SNew(SHorizontalBox)
+				 +SHorizontalBox::Slot()
+			    .AutoWidth()
+			    .Padding(1.0f, 0.0f)
+			    [
+				    SNew(SBorder)
+				    .Padding(0.0f)
+				    .BorderImage(StatusBarBackground)
+				    .VAlign(VAlign_Center)
+				    .Padding(FMargin(6.0f, 0.0f))
+				    [
+					    MakeDebugConsoleWidget(InArgs._OnConsoleClosed)
+				    ]
+			    ]
+				+SHorizontalBox::Slot()
+				.FillWidth(1.0f)
+				.Padding(1.0f, 0.0f)
 				[
-					MakeStatusMessageWidget()
+					SNew(SBorder)
+					.Padding(0.0f)
+					.BorderImage(StatusBarBackground)
+					.VAlign(VAlign_Center)
+					.Padding(FMargin(6.0f, 0.0f))
+					[
+						MakeStatusMessageWidget()
+					]
 				]
-			]
-			+ SHorizontalBox::Slot()
-			.HAlign(HAlign_Right)
-			.AutoWidth()
-			.Padding(1.0f, 0.0f)
-			[
-				SNew(SBorder)
-				.Padding(0.0f)
-				.BorderImage(StatusBarBackground)
-				.VAlign(VAlign_Center)
+				+SHorizontalBox::Slot()
+				.HAlign(HAlign_Right)
+				.AutoWidth()
+				.Padding(1.0f, 0.0f)
 				[
-					MakeStatusBarToolBarWidget()
+					SNew(SBorder)
+					.Padding(0.0f)
+					.BorderImage(StatusBarBackground)
+					.VAlign(VAlign_Center)
+					[
+						MakeStatusBarToolBarWidget()
+					]
 				]
 			]
 		]
@@ -407,10 +407,15 @@ void SStatusBar::OnGlobalFocusChanging(const FFocusEvent& FocusEvent, const FWea
 
 		TSharedRef<SWidget> ThisWidget = AsShared();
 
-		
+		TSharedPtr<SWidget> ActiveDrawerOverlayContent;
+		if (!OpenedDrawerData.Key.IsNone())
+		{
+			ActiveDrawerOverlayContent = OpenedDrawerData.Value;
+		}
+
 		bool bShouldDismiss = false;
 		// Do not close due to slow tasks as those opening send window activation events
-		if (!GIsSlowTask && !FSlateApplication::Get().GetActiveModalWindow().IsValid() && ContentBrowserOverlayContent.IsValid() && (!NewFocusedWidgetPath.ContainsWidget(ContentBrowserOverlayContent.ToSharedRef()) && !NewFocusedWidgetPath.ContainsWidget(ThisWidget)))
+		if (!GIsSlowTask && !FSlateApplication::Get().GetActiveModalWindow().IsValid() && ActiveDrawerOverlayContent.IsValid() && (!NewFocusedWidgetPath.ContainsWidget(ActiveDrawerOverlayContent.ToSharedRef()) && !NewFocusedWidgetPath.ContainsWidget(ThisWidget)))
 		{
 			if (TSharedPtr<SWidget> MenuHost = FSlateApplication::Get().GetMenuHostWidget())
 			{
@@ -418,7 +423,7 @@ void SStatusBar::OnGlobalFocusChanging(const FFocusEvent& FocusEvent, const FWea
 
 				// See if the menu being opened is part of the content browser path and if so the menu should not be dismissed
 				FSlateApplication::Get().GeneratePathToWidgetUnchecked(MenuHost.ToSharedRef(), MenuHostPath, EVisibility::All);
-				if (!MenuHostPath.ContainsWidget(ContentBrowserOverlayContent.ToSharedRef()))
+				if (!MenuHostPath.ContainsWidget(ActiveDrawerOverlayContent.ToSharedRef()))
 				{
 					bShouldDismiss = true;
 				}
@@ -431,141 +436,36 @@ void SStatusBar::OnGlobalFocusChanging(const FFocusEvent& FocusEvent, const FWea
 
 		if (bShouldDismiss)
 		{
-			DismissContentBrowser(NewFocusedWidget);
+			DismissDrawer(NewFocusedWidget);
 		}
 	}
 }
 
 void SStatusBar::OnActiveTabChanged(TSharedPtr<SDockTab> PreviouslyActive, TSharedPtr<SDockTab> NewlyActivated)
 {
-	bool bShouldRemoveContentBrowser = false;
+	bool bShouldRemoveDrawer = false;
 	if (!PreviouslyActive || !NewlyActivated)
 	{
 		// Remove the content browser if there is some invalid state with the tabs
-		bShouldRemoveContentBrowser = true;
+		bShouldRemoveDrawer = true;
 	}
 	else if(NewlyActivated->GetTabRole() == ETabRole::MajorTab)
 	{
 		// Remove the content browser if a newly activated tab is a major tab
-		bShouldRemoveContentBrowser = true;
+		bShouldRemoveDrawer = true;
 	}
 	else if (PreviouslyActive->GetTabManager() != NewlyActivated->GetTabManager())
 	{
 		// Remove the content browser if we're switching tab managers (indicates a new status bar is becoming active)
-		bShouldRemoveContentBrowser = true;
+		bShouldRemoveDrawer = true;
 	}
 
-	if (bShouldRemoveContentBrowser)
+	if (bShouldRemoveDrawer)
 	{
-		RemoveContentBrowser();
+		CloseDrawerImmediately();
 	}
 }
 
-bool SStatusBar::IsContentBrowserOpened() const
-{
-	return ContentBrowserOverlayContent.IsValid();
-}
-
-void SStatusBar::OpenContentBrowser()
-{
-	if (!IsContentBrowserOpened())
-	{
-		TSharedRef<SStatusBar> ThisStatusBar = SharedThis(this);
-
-		TSharedPtr<SWindow> MyWindow = FSlateApplication::Get().FindWidgetWindow(AsShared());
-
-		const float MaxContentBrowserHeight = MyWindow->GetSizeInScreen().Y * 0.90f;
-
-		float TargetContentBrowserHeightPct = .33f;
-		GConfig->GetFloat(TEXT("ContentBrowserDrawerSizes"), *StatusBarName.ToString(), TargetContentBrowserHeightPct, GEditorSettingsIni);
-
-		TargetContentBrowserHeight = (MyWindow->GetSizeInScreen().Y * TargetContentBrowserHeightPct) / MyWindow->GetDPIScaleFactor();
-
-
-		const float MinContentBrowserHeight = GetTickSpaceGeometry().GetLocalSize().Y + MyWindow->GetWindowBorderSize().Bottom;
-		ContentBrowserEasingCurve.Play(ThisStatusBar, false, ContentBrowserEasingCurve.IsPlaying() ? ContentBrowserEasingCurve.GetSequenceTime() : 0.0f, false);
-
-		MyWindow->AddOverlaySlot()
-			.VAlign(VAlign_Bottom)
-			.Padding(FMargin(10.0f, 20.0f, 10.0f, MinContentBrowserHeight))
-			[
-				SAssignNew(ContentBrowserOverlayContent, SContentBrowserOverlay)
-				.MinContentBrowserHeight(MinContentBrowserHeight)
-				.TargetContentBrowserHeight(TargetContentBrowserHeight)
-				.MaxContentBrowserHeight(MaxContentBrowserHeight)
-				.OnTargetHeightChanged(this, &SStatusBar::OnContentBrowserTargetHeightChanged)
-				[
-					GetContentBrowserDelegate.Execute()
-				]
-			];
-
-		WindowWithOverlayContent = MyWindow;
-
-		if (!ContentBrowserOpenCloseTimer.IsValid())
-		{
-			AnimationThrottle = FSlateThrottleManager::Get().EnterResponsiveMode();
-			ContentBrowserOpenCloseTimer = RegisterActiveTimer(0.0f, FWidgetActiveTimerDelegate::CreateSP(ThisStatusBar, &SStatusBar::UpdateContentBrowserAnimation));
-		}
-
-		OnContentBrowserOpenedDelegate.ExecuteIfBound(ThisStatusBar);
-	}
-}
-
-/**
- * Dismisses an open content browser with an animation.  The content browser is removed once the animation is complete
- */
-void SStatusBar::DismissContentBrowser(const TSharedPtr<SWidget>& NewlyFocusedWidget)
-{
-	if (ContentBrowserOverlayContent.IsValid())
-	{
-		if (ContentBrowserEasingCurve.IsForward())
-		{
-			ContentBrowserEasingCurve.Reverse();
-		}
-
-		if (!ContentBrowserOpenCloseTimer.IsValid())
-		{
-			AnimationThrottle = FSlateThrottleManager::Get().EnterResponsiveMode();
-			ContentBrowserOpenCloseTimer = RegisterActiveTimer(0.0f, FWidgetActiveTimerDelegate::CreateSP(this, &SStatusBar::UpdateContentBrowserAnimation));
-		}
-
-		OnContentBrowserDismissedDelegate.ExecuteIfBound(NewlyFocusedWidget);
-	}
-}
-
-/**
- * Removes a content browser from the layout immediately
- */
-void SStatusBar::RemoveContentBrowser()
-{
-	if(ContentBrowserOverlayContent.IsValid())
-	{
-		// Remove the content browser from the window
-		if (TSharedPtr<SWindow> Window = WindowWithOverlayContent.Pin())
-		{
-			Window->RemoveOverlaySlot(ContentBrowserOverlayContent.ToSharedRef());
-		}
-
-		ContentBrowserOverlayContent.Reset();
-		WindowWithOverlayContent.Reset();
-
-		if (ContentBrowserOpenCloseTimer.IsValid())
-		{
-			FSlateThrottleManager::Get().LeaveResponsiveMode(AnimationThrottle);
-			UnRegisterActiveTimer(ContentBrowserOpenCloseTimer.ToSharedRef());
-		}
-
-		ContentBrowserOpenCloseTimer.Reset();
-		ContentBrowserEasingCurve.JumpToStart();
-		ContentBrowserEasingCurve.Pause();
-	}
-}
-
-const FSlateBrush* SStatusBar::GetContentBrowserExpandArrowImage() const
-{
-	return ContentBrowserOverlayContent.IsValid() ? DownArrow : UpArrow;
-}
-
 FText SStatusBar::GetStatusBarMessage() const
 {
 	FText FullMessage;
@@ -582,14 +482,16 @@ FText SStatusBar::GetStatusBarMessage() const
 	return FullMessage;
 }
 
-TSharedRef<SWidget> SStatusBar::MakeContentBrowserWidget()
+TSharedRef<SWidget> SStatusBar::MakeStatusBarDrawerButton(const FStatusBarDrawer& Drawer)
 {
+	const FName DrawerId = Drawer.UniqueId;
+
 	return
 		SNew(SButton)
 		.IsFocusable(false)
 		.ButtonStyle(&FAppStyle::Get().GetWidgetStyle<FButtonStyle>("StatusBar.StatusBarButton"))
-		.OnClicked(this, &SStatusBar::OnContentBrowserButtonClicked)
-		.ToolTipText(FText::Format(LOCTEXT("StatusBar_ContentBrowserDrawerToolTip","Opens a temporary content browser above this status which will dismiss when it loses focus ({0})"), FGlobalEditorCommonCommands::Get().OpenContentBrowserDrawer->GetInputText()))
+		.OnClicked(this, &SStatusBar::OnDrawerButtonClicked, DrawerId)
+		.ToolTipText(Drawer.ToolTipText)
 		[
 			SNew(SHorizontalBox)
 			+ SHorizontalBox::Slot()
@@ -600,7 +502,7 @@ TSharedRef<SWidget> SStatusBar::MakeContentBrowserWidget()
 			[
 				SNew(SImage)
 				.ColorAndOpacity(FSlateColor::UseForeground())
-				.Image(this, &SStatusBar::GetContentBrowserExpandArrowImage)
+				.Image_Lambda([this, DrawerId](){ return IsDrawerOpened(DrawerId) ? DownArrow : UpArrow; })
 			]
 			+ SHorizontalBox::Slot()
 			.Padding(2.0f)
@@ -610,7 +512,7 @@ TSharedRef<SWidget> SStatusBar::MakeContentBrowserWidget()
 			[
 				SNew(SImage)
 				.ColorAndOpacity(FSlateColor::UseForeground())
-				.Image(FAppStyle::Get().GetBrush("ContentBrowser.TabIcon"))
+				.Image(Drawer.Icon)
 			]
 			+ SHorizontalBox::Slot()
 			.VAlign(VAlign_Center)
@@ -618,7 +520,7 @@ TSharedRef<SWidget> SStatusBar::MakeContentBrowserWidget()
 			[
 				SNew(STextBlock)
 				.TextStyle(&FAppStyle::Get().GetWidgetStyle<FTextBlockStyle>("NormalText"))
-				.Text(LOCTEXT("StatusBar_ContentBrowserButton", "Content Browser"))
+				.Text(Drawer.ButtonText)
 			]
 		];	
 }
@@ -670,52 +572,56 @@ TSharedRef<SWidget> SStatusBar::MakeStatusMessageWidget()
 		];
 }
 
-FReply SStatusBar::OnContentBrowserButtonClicked()
+
+bool SStatusBar::IsDrawerOpened(const FName DrawerId) const
+{
+	return OpenedDrawerData.Key == DrawerId ? true : false;
+}
+
+FReply SStatusBar::OnDrawerButtonClicked(const FName DrawerId)
 {
-	if(!ContentBrowserOverlayContent.IsValid())
+	if (!IsDrawerOpened(DrawerId))
 	{
-		OpenContentBrowser();
+		OpenDrawer(DrawerId);
 	}
 	else
 	{
-		DismissContentBrowser(nullptr);
+		DismissDrawer(nullptr);
 	}
 
 	return FReply::Handled();
 }
 
-EActiveTimerReturnType SStatusBar::UpdateContentBrowserAnimation(double CurrentTime, float DeltaTime)
+
+EActiveTimerReturnType SStatusBar::UpdateDrawerAnimation(double CurrentTime, float DeltaTime)
 {
-	check(ContentBrowserOverlayContent.IsValid());
+	check(OpenedDrawerData.Value.IsValid());
 
-	ContentBrowserOverlayContent->SetHeight(FMath::Lerp(0.0f, TargetContentBrowserHeight, ContentBrowserEasingCurve.GetLerp()));
+	OpenedDrawerData.Value->UpdateHeightInterp(DrawerEasingCurve.GetLerp());
 
-	if (!ContentBrowserEasingCurve.IsPlaying())
+	if (!DrawerEasingCurve.IsPlaying())
 	{
-		if(ContentBrowserEasingCurve.IsAtStart())
+		if(DrawerEasingCurve.IsAtStart())
 		{
-			RemoveContentBrowser();
+			CloseDrawerImmediately();
 		}
 
 		FSlateThrottleManager::Get().LeaveResponsiveMode(AnimationThrottle);
-		ContentBrowserOpenCloseTimer.Reset();
+		DrawerOpenCloseTimer.Reset();
 		return EActiveTimerReturnType::Stop;
 	}
 
 	return EActiveTimerReturnType::Continue;
 }
 
-void SStatusBar::OnContentBrowserTargetHeightChanged(float TargetHeight)
+void SStatusBar::OnDrawerHeightChanged(float TargetHeight)
 {
-	TargetContentBrowserHeight = TargetHeight;
-
 	TSharedPtr<SWindow> MyWindow = WindowWithOverlayContent.Pin();
 
 	// Save the height has a percentage of the screen
-	const float TargetContentBrowserHeightPct = TargetHeight / (MyWindow->GetSizeInScreen().Y / MyWindow->GetDPIScaleFactor());
+	const float TargetDrawerHeightPct = TargetHeight / (MyWindow->GetSizeInScreen().Y / MyWindow->GetDPIScaleFactor());
 
-	// Write to config
-	GConfig->SetFloat(TEXT("ContentBrowserDrawerSizes"), *StatusBarName.ToString(), TargetContentBrowserHeightPct, GEditorSettingsIni);
+	GConfig->SetFloat(TEXT("DrawerSizes"), *(StatusBarName.ToString() + TEXT(".") + OpenedDrawerData.Key.ToString()), TargetDrawerHeightPct, GEditorSettingsIni);
 }
 
 void SStatusBar::RegisterStatusBarMenu()
@@ -749,4 +655,159 @@ void SStatusBar::RegisterSourceControlStatus()
 	}
 }
 
+void SStatusBar::RegisterDrawer(FStatusBarDrawer&& Drawer)
+{
+	const int32 NumDrawers = RegisteredDrawers.Num();
+	RegisteredDrawers.AddUnique(Drawer);
+
+	if (RegisteredDrawers.Num() > NumDrawers)
+	{
+		const FSlateBrush* StatusBarBackground = FAppStyle::Get().GetBrush("StatusBar.Background");
+
+		DrawerBox->AddSlot()
+		.Padding(1.0f, 0.0f)
+		.AutoWidth()
+		[
+			SNew(SBorder)
+			.Padding(FMargin(2.0f, 0.0f))
+			.BorderImage(StatusBarBackground)
+			.VAlign(VAlign_Center)
+			[
+				MakeStatusBarDrawerButton(Drawer)
+			]
+		];
+/*
+			/ *+ SHorizontalBox::Slot()
+				.Padding(1.0f, 0.0f)
+				.AutoWidth()
+				[
+					SNew(SBorder)
+					.Padding(FMargin(2.0f, 0.0f))
+					.BorderImage(StatusBarBackground)
+					.VAlign(VAlign_Center)
+					[
+						MakeStatusBarDrawerButton(ContentBrowserDrawerButton)
+					]
+				]
+				+ SHorizontalBox::Slot()
+				.AutoWidth()
+				.Padding(1.0f, 0.0f)
+				[
+					SNew(SBorder)
+					.Padding(0.0f)
+					.BorderImage(StatusBarBackground)
+					.VAlign(VAlign_Center)
+					.Padding(FMargin(6.0f, 0.0f))
+					[
+						MakeDebugConsoleWidget(InArgs._OnConsoleClosed)
+					]
+				]* /
+		}
+*/
+
+	}
+}
+
+void SStatusBar::OpenDrawer(const FName DrawerId)
+{
+	// Close any other open drawer
+	if (DrawerId != OpenedDrawerData.Key)
+	{
+		DismissDrawer(nullptr);
+
+		FStatusBarDrawer* DrawerData = RegisteredDrawers.FindByKey(DrawerId);
+
+		if(DrawerData)
+		{
+			TSharedRef<SStatusBar> ThisStatusBar = SharedThis(this);
+
+			TSharedPtr<SWindow> MyWindow = FSlateApplication::Get().FindWidgetWindow(AsShared());
+
+			const float MaxDrawerHeight = MyWindow->GetSizeInScreen().Y * 0.90f;
+
+			float TargetDrawerHeightPct = .33f;
+			GConfig->GetFloat(TEXT("DrawerSizes"), *(StatusBarName.ToString()+TEXT(".")+DrawerData->UniqueId.ToString()), TargetDrawerHeightPct, GEditorSettingsIni);
+
+			float TargetDrawerHeight = (MyWindow->GetSizeInScreen().Y * TargetDrawerHeightPct) / MyWindow->GetDPIScaleFactor();
+
+			const float MinDrawerHeight = GetTickSpaceGeometry().GetLocalSize().Y + MyWindow->GetWindowBorderSize().Bottom;
+			DrawerEasingCurve.Play(ThisStatusBar, false, DrawerEasingCurve.IsPlaying() ? DrawerEasingCurve.GetSequenceTime() : 0.0f, false);
+
+			MyWindow->AddOverlaySlot()
+				.VAlign(VAlign_Bottom)
+				.Padding(FMargin(10.0f, 20.0f, 10.0f, MinDrawerHeight))
+				[
+					SAssignNew(OpenedDrawerData.Value, SDrawerOverlay)
+					.MinDrawerHeight(MinDrawerHeight)
+					.TargetDrawerHeight(TargetDrawerHeight)
+					.MaxDrawerHeight(MaxDrawerHeight)
+					.OnTargetHeightChanged(this, &SStatusBar::OnDrawerHeightChanged)
+					[
+						DrawerData->GetDrawerContentDelegate.Execute()
+					]
+				];
+
+			WindowWithOverlayContent = MyWindow;
+
+			if (!DrawerOpenCloseTimer.IsValid())
+			{
+				AnimationThrottle = FSlateThrottleManager::Get().EnterResponsiveMode();
+				DrawerOpenCloseTimer = RegisterActiveTimer(0.0f, FWidgetActiveTimerDelegate::CreateSP(ThisStatusBar, &SStatusBar::UpdateDrawerAnimation));
+			}
+
+			OpenedDrawerData.Key = DrawerId;
+			DrawerData->OnDrawerOpenedDelegate.ExecuteIfBound(ThisStatusBar);
+		}
+	}
+}
+
+void SStatusBar::DismissDrawer(const TSharedPtr<SWidget>& NewlyFocusedWidget)
+{
+	if (!OpenedDrawerData.Key.IsNone())
+	{
+		FStatusBarDrawer* Drawer = RegisteredDrawers.FindByKey(OpenedDrawerData.Key);
+
+		if (DrawerEasingCurve.IsForward())
+		{
+			DrawerEasingCurve.Reverse();
+		}
+
+		if (!DrawerOpenCloseTimer.IsValid())
+		{
+			AnimationThrottle = FSlateThrottleManager::Get().EnterResponsiveMode();
+			DrawerOpenCloseTimer = RegisterActiveTimer(0.0f, FWidgetActiveTimerDelegate::CreateSP(this, &SStatusBar::UpdateDrawerAnimation));
+		}
+
+		Drawer->OnDrawerDismissedDelegate.ExecuteIfBound(NewlyFocusedWidget);
+	}
+}
+
+void SStatusBar::CloseDrawerImmediately()
+{
+	if (!OpenedDrawerData.Key.IsNone())
+	{
+		TSharedRef<SWidget> DrawerOverlayContent = OpenedDrawerData.Value.ToSharedRef();
+
+		// Remove the content browser from the window
+		if (TSharedPtr<SWindow> Window = WindowWithOverlayContent.Pin())
+		{
+			Window->RemoveOverlaySlot(DrawerOverlayContent);
+		}
+
+		OpenedDrawerData = TPair<FName, TSharedPtr<SDrawerOverlay>>();
+
+		WindowWithOverlayContent.Reset();
+
+		if (DrawerOpenCloseTimer.IsValid())
+		{
+			FSlateThrottleManager::Get().LeaveResponsiveMode(AnimationThrottle);
+			UnRegisterActiveTimer(DrawerOpenCloseTimer.ToSharedRef());
+		}
+
+		DrawerOpenCloseTimer.Reset();
+		DrawerEasingCurve.JumpToStart();
+		DrawerEasingCurve.Pause();
+	}
+}
+
 #undef LOCTEXT_NAMESPACE
\ No newline at end of file
diff --git a/Engine/Source/Editor/StatusBar/Private/StatusBarSubsystem.cpp b/Engine/Source/Editor/StatusBar/Private/StatusBarSubsystem.cpp
--- a/Engine/Source/Editor/StatusBar/Private/StatusBarSubsystem.cpp
+++ b/Engine/Source/Editor/StatusBar/Private/StatusBarSubsystem.cpp
@@ -11,6 +11,7 @@
 #include "SStatusBar.h"
 #include "Widgets/Docking/SDockTab.h"
 #include "Editor.h"
+#include "Toolkits/GlobalEditorCommonCommands.h"
 
 
 #define LOCTEXT_NAMESPACE "StatusBar"
@@ -81,12 +82,12 @@ bool UStatusBarSubsystem::ToggleContentBrowser(TSharedRef<SWindow> ParentWindow)
 	{
 		if (TSharedPtr<SStatusBar> StatusBarPinned = StatusBar.Value.Pin())
 		{
-			if(StatusBarPinned->IsContentBrowserOpened())
+			if(StatusBarPinned->IsDrawerOpened(StatusBarDrawerIds::ContentBrowser))
 			{
 				TSharedPtr<SDockTab> ParentTab = StatusBarPinned->GetParentTab();
 				if (ParentTab && ParentTab->IsForeground() && ParentTab->GetParentWindow() == ParentWindow)
 				{
-					StatusBarPinned->DismissContentBrowser(nullptr);
+					StatusBarPinned->DismissDrawer(nullptr);
 					bWasDismissed = true;
 				}
 			}
@@ -108,10 +109,17 @@ TSharedRef<SWidget> UStatusBarSubsystem::MakeStatusBarWidget(FName StatusBarName
 
 	TSharedRef<SStatusBar> StatusBar =
 		SNew(SStatusBar, StatusBarName, InParentTab)
-		.OnConsoleClosed_UObject(this, &UStatusBarSubsystem::OnDebugConsoleClosed)
-		.OnGetContentBrowser_UObject(this, &UStatusBarSubsystem::OnGetContentBrowser)
-		.OnContentBrowserOpened_UObject(this, &UStatusBarSubsystem::OnContentBrowserOpened)
-		.OnContentBrowserDismissed_UObject(this, &UStatusBarSubsystem::OnContentBrowserDismissed);
+		.OnConsoleClosed_UObject(this, &UStatusBarSubsystem::OnDebugConsoleClosed);
+
+	FStatusBarDrawer ContentBrowserDrawer(StatusBarDrawerIds::ContentBrowser);
+	ContentBrowserDrawer.GetDrawerContentDelegate.BindUObject(this, &UStatusBarSubsystem::OnGetContentBrowser);
+	ContentBrowserDrawer.OnDrawerOpenedDelegate.BindUObject(this, &UStatusBarSubsystem::OnContentBrowserOpened);
+	ContentBrowserDrawer.OnDrawerDismissedDelegate.BindUObject(this, &UStatusBarSubsystem::OnContentBrowserDismissed);
+	ContentBrowserDrawer.ButtonText = LOCTEXT("StatusBar_ContentBrowserButton", "Content Browser");
+	ContentBrowserDrawer.ToolTipText = FText::Format(LOCTEXT("StatusBar_ContentBrowserDrawerToolTip", "Opens a temporary content browser above this status which will dismiss when it loses focus ({0})"), FGlobalEditorCommonCommands::Get().OpenContentBrowserDrawer->GetInputText());
+	ContentBrowserDrawer.Icon = FAppStyle::Get().GetBrush("ContentBrowser.TabIcon");
+
+	StatusBar->RegisterDrawer(MoveTemp(ContentBrowserDrawer));
 
 	// Clean up stale status bars
 	for (auto It = StatusBars.CreateIterator(); It; ++It)
@@ -203,9 +211,9 @@ void UStatusBarSubsystem::OnContentBrowserOpened(TSharedRef<SStatusBar>& StatusB
 	{
 		if (TSharedPtr<SStatusBar> StatusBarPinned = StatusBar.Value.Pin())
 		{
-			if (StatusBarWithContentBrowser != StatusBarPinned && StatusBarPinned->IsContentBrowserOpened())
+			if (StatusBarWithContentBrowser != StatusBarPinned && StatusBarPinned->IsDrawerOpened(StatusBarDrawerIds::ContentBrowser))
 			{
-				StatusBarPinned->RemoveContentBrowser();
+				StatusBarPinned->CloseDrawerImmediately();
 			}
 		}
 	}
@@ -240,7 +248,7 @@ void UStatusBarSubsystem::HandleDeferredOpenContentBrowser(TSharedPtr<SWindow> P
 			TSharedPtr<SDockTab> ParentTab = StatusBarPinned->GetParentTab();
 			if (ParentTab && ParentTab->IsForeground() && ParentTab->GetParentWindow() == ParentWindow)
 			{
-				StatusBarPinned->OpenContentBrowser();
+				StatusBarPinned->OpenDrawer(StatusBarDrawerIds::ContentBrowser);
 				IContentBrowserSingleton& ContentBrowserSingleton = FModuleManager::Get().GetModuleChecked<FContentBrowserModule>("ContentBrowser").Get();
 				break;
 			}
diff --git a/Engine/Source/Editor/StatusBar/Public/SStatusBar.h b/Engine/Source/Editor/StatusBar/Public/SStatusBar.h
--- a/Engine/Source/Editor/StatusBar/Public/SStatusBar.h
+++ b/Engine/Source/Editor/StatusBar/Public/SStatusBar.h
@@ -10,30 +10,65 @@
 #include "Application/ThrottleManager.h"
 
 class SStatusBar;
-
-DECLARE_DELEGATE_OneParam(FOnContentBrowserOpened, TSharedRef<SStatusBar>&)
-DECLARE_DELEGATE_OneParam(FOnContentBrowserDismissed, const TSharedPtr<SWidget>&)
-DECLARE_DELEGATE_OneParam(FOnContentBrowserTargetHeightChanged, float);
-
+DECLARE_DELEGATE_OneParam(FOnStatusBarDrawerOpened, TSharedRef<SStatusBar>&)
+DECLARE_DELEGATE_OneParam(FOnStatusBarDrawerDismissed, const TSharedPtr<SWidget>&)
+DECLARE_DELEGATE_OneParam(FOnStatusBarDrawerTargetHeightChanged, float);
 
 class SWindow;
 class SDockTab;
 class SWidget;
 class SMultiLineEditableTextBox;
-class SContentBrowserOverlay;
+class SDrawerOverlay;
 
-class SStatusBar : public SCompoundWidget
+namespace StatusBarDrawerIds
 {
-	SLATE_BEGIN_ARGS(SStatusBar)	
+	const FName ContentBrowser("ContentBrowser");
+}
+
+struct FStatusBarMessage
+{
+	FStatusBarMessage(const TAttribute<FText>& InMessageText, const TAttribute<FText>& InHintText, FStatusBarMessageHandle InHandle)
+		: MessageText(InMessageText)
+		, HintText(InHintText)
+		, Handle(InHandle)
 	{}
-		SLATE_EVENT(FSimpleDelegate, OnConsoleClosed)
 
-		SLATE_EVENT(FOnGetContent, OnGetContentBrowser)
+	TAttribute<FText> MessageText;
+	TAttribute<FText> HintText;
+	FStatusBarMessageHandle Handle;
+};
 
-		SLATE_EVENT(FOnContentBrowserOpened, OnContentBrowserOpened)
+struct FStatusBarDrawer
+{
+	FStatusBarDrawer(FName InUniqueId)
+		: UniqueId(InUniqueId)
+	{}
+
+	FName UniqueId;
+	FOnGetContent GetDrawerContentDelegate;
+	FOnStatusBarDrawerOpened OnDrawerOpenedDelegate;
+	FOnStatusBarDrawerDismissed OnDrawerDismissedDelegate;
+
+	FText ButtonText;
+	FText ToolTipText;
+	const FSlateBrush* Icon;
+
+	bool operator==(const FName& OtherId) const
+	{
+		return UniqueId == OtherId;
+	}
 
-		SLATE_EVENT(FOnContentBrowserDismissed, OnContentBrowserDismissed)
+	bool operator==(const FStatusBarDrawer& Other) const
+	{
+		return UniqueId == Other.UniqueId;
+	}
+};
 
+class SStatusBar : public SCompoundWidget
+{
+	SLATE_BEGIN_ARGS(SStatusBar)	
+	{}
+		SLATE_EVENT(FSimpleDelegate, OnConsoleClosed)
 	SLATE_END_ARGS()
 
 public:
@@ -68,19 +103,18 @@ public:
 	bool FocusDebugConsole();
 	bool IsDebugConsoleFocused() const;
 
-	bool IsContentBrowserOpened() const;
+	void RegisterDrawer(FStatusBarDrawer&& Drawer);
 
-	void OpenContentBrowser();
+	void OpenDrawer(const FName DrawerId);
 
 	/**
-	 * Dismisses an open content browser with an animation.  The content browser is removed once the animation is complete
+	 * Dismisses an open drawer with an animation.  The drawer contents are removed once the animation is complete
 	 */
-	void DismissContentBrowser(const TSharedPtr<SWidget>& NewlyFocusedWidget);
+	void DismissDrawer(const TSharedPtr<SWidget>& NewlyFocusedWidget);
 
-	/**
-	 * Removes a content browser from the layout immediately
-	 */
-	void RemoveContentBrowser();
+	void CloseDrawerImmediately();
+
+	bool IsDrawerOpened(const FName DrawerId) const;
 private:
 	/** Called when global focus changes which is used to determine if we should close an opened content browser drawer */
 	void OnGlobalFocusChanging(const FFocusEvent& FocusEvent, const FWeakWidgetPath& OldFocusedWidgetPath, const TSharedPtr<SWidget>& OldFocusedWidget, const FWidgetPath& NewFocusedWidgetPath, const TSharedPtr<SWidget>& NewFocusedWidget);
@@ -88,54 +122,47 @@ private:
 	/** Called when active tab changes which is used to determine if we should close an opened content browser drawer */
 	void OnActiveTabChanged(TSharedPtr<SDockTab> PreviouslyActive, TSharedPtr<SDockTab> NewlyActivated);
 
+	void OnDrawerHeightChanged(float TargetHeight);
+
 	EVisibility GetHelpIconVisibility() const;
 
-	const FSlateBrush* GetContentBrowserExpandArrowImage() const;
+	bool GetContentBrowserExpanded() const;
 
 	FText GetStatusBarMessage() const;
 
-	TSharedRef<SWidget> MakeContentBrowserWidget();
+	TSharedRef<SWidget> MakeStatusBarDrawerButton(const FStatusBarDrawer& Drawer);
 
 	TSharedRef<SWidget> MakeStatusBarToolBarWidget();
 	TSharedRef<SWidget> MakeDebugConsoleWidget(FSimpleDelegate OnConsoleClosed);
 	TSharedRef<SWidget> MakeStatusMessageWidget();
+	void MakeDrawerButtons(TSharedRef<SHorizontalBox> DrawerBox);
 
-	FReply OnContentBrowserButtonClicked();
+
+	FReply OnDrawerButtonClicked(const FName DrawerId);
 
 	void CreateAnimationTimerIfNeeded();
 
-	EActiveTimerReturnType UpdateContentBrowserAnimation(double CurrentTime, float DeltaTime);
-	void OnContentBrowserTargetHeightChanged(float TargetHeight);
+	EActiveTimerReturnType UpdateDrawerAnimation(double CurrentTime, float DeltaTime);
 
 	void RegisterStatusBarMenu();
 
 	void RegisterSourceControlStatus();
-
-	struct FStatusBarMessage
-	{
-		FStatusBarMessage(const TAttribute<FText>& InMessageText, const TAttribute<FText>& InHintText, FStatusBarMessageHandle InHandle)
-			: MessageText(InMessageText)
-			, HintText(InHintText)
-			, Handle(InHandle)
-		{}
-
-		TAttribute<FText> MessageText;
-		TAttribute<FText> HintText;
-		FStatusBarMessageHandle Handle;
-	};
 private:
 	TArray<FStatusBarMessage> MessageStack;
 	TSharedPtr<SMultiLineEditableTextBox> ConsoleEditBox;
 	TWeakPtr<SDockTab> ParentTab;
-	TSharedPtr<SContentBrowserOverlay> ContentBrowserOverlayContent;
+
+	TArray<FStatusBarDrawer> RegisteredDrawers;
+
 	TWeakPtr<SWindow> WindowWithOverlayContent;
-	TSharedPtr<FActiveTimerHandle> ContentBrowserOpenCloseTimer;
-	FOnGetContent GetContentBrowserDelegate;
-	FOnContentBrowserOpened OnContentBrowserOpenedDelegate;
-	FOnContentBrowserDismissed OnContentBrowserDismissedDelegate;
-	FCurveSequence ContentBrowserEasingCurve;
+
+	TPair<FName,TSharedPtr<SDrawerOverlay>> OpenedDrawerData;
+
+	TSharedPtr<SHorizontalBox> DrawerBox;
+
+	TSharedPtr<FActiveTimerHandle> DrawerOpenCloseTimer;
+	FCurveSequence DrawerEasingCurve;
 	FThrottleRequest AnimationThrottle;
-	float TargetContentBrowserHeight;
 	const FSlateBrush* UpArrow;
 	const FSlateBrush* DownArrow;
 	FName StatusBarName;
diff --git a/Engine/Source/Editor/UnrealEd/Private/Toolkits/GlobalEditorCommonCommands.cpp b/Engine/Source/Editor/UnrealEd/Private/Toolkits/GlobalEditorCommonCommands.cpp
--- a/Engine/Source/Editor/UnrealEd/Private/Toolkits/GlobalEditorCommonCommands.cpp
+++ b/Engine/Source/Editor/UnrealEd/Private/Toolkits/GlobalEditorCommonCommands.cpp
@@ -134,6 +134,16 @@ static void CloseDebugConsole()
 void FGlobalEditorCommonCommands::OnSummonedConsoleCommandBox()
 {
 	TSharedPtr<SWindow> ParentWindow = FSlateApplication::Get().GetActiveTopLevelWindow();
+	if (!ParentWindow.IsValid())
+	{
+		if (TSharedPtr<SDockTab> ActiveTab = FGlobalTabmanager::Get()->GetActiveTab())
+		{
+			if (TSharedPtr<SDockTab> ActiveMajorTab = FGlobalTabmanager::Get()->GetMajorTabForTabManager(ActiveTab->GetTabManager()))
+			{
+				ParentWindow = ActiveMajorTab->GetParentWindow();
+			}
+		}
+	}
 
 	if (ParentWindow.IsValid() && ParentWindow->GetType() == EWindowType::Normal)
 	{
