Fix Non-Unity CIS

diff --git a/Engine/Source/Editor/StatusBar/Public/SStatusBar.h b/Engine/Source/Editor/StatusBar/Public/SStatusBar.h
--- a/Engine/Source/Editor/StatusBar/Public/SStatusBar.h
+++ b/Engine/Source/Editor/StatusBar/Public/SStatusBar.h
@@ -10,6 +10,7 @@
 #include "Application/ThrottleManager.h"
 
 class SStatusBar;
+
 DECLARE_DELEGATE_OneParam(FOnStatusBarDrawerOpened, TSharedRef<SStatusBar>&)
 DECLARE_DELEGATE_OneParam(FOnStatusBarDrawerDismissed, const TSharedPtr<SWidget>&)
 DECLARE_DELEGATE_OneParam(FOnStatusBarDrawerTargetHeightChanged, float);
@@ -19,6 +20,7 @@ class SDockTab;
 class SWidget;
 class SMultiLineEditableTextBox;
 class SDrawerOverlay;
+class SHorizontalBox;
 
 namespace StatusBarDrawerIds
 {
