SourceControl: Replace main menu "Submit" menu with "View Changelists..."

diff --git a/Engine/Source/Editor/MainFrame/Private/Frame/MainFrameActions.cpp b/Engine/Source/Editor/MainFrame/Private/Frame/MainFrameActions.cpp
--- a/Engine/Source/Editor/MainFrame/Private/Frame/MainFrameActions.cpp
+++ b/Engine/Source/Editor/MainFrame/Private/Frame/MainFrameActions.cpp
@@ -18,12 +18,12 @@
 #include "CreditsScreen.h"
 #include "DesktopPlatformModule.h"
 #include "ISourceControlModule.h"
+#include "ISourceControlWindowsModule.h"
 #include "GameProjectGenerationModule.h"
 #include "Toolkits/GlobalEditorCommonCommands.h"
 #include "Logging/TokenizedMessage.h"
 #include "Logging/MessageLog.h"
 #include "SourceCodeNavigation.h"
-#include "SourceControlWindows.h"
 #include "ISettingsModule.h"
 #include "Interfaces/IProjectManager.h"
 #include "Interfaces/ITargetPlatform.h"
@@ -99,8 +99,8 @@ void FMainFrameCommands::RegisterCommands()
 	UI_COMMAND( ChooseFilesToSave, "Choose Files to Save...", "Opens a dialog with save options for content and levels", EUserInterfaceActionType::Button, FInputChord() );
 	ActionList->MapAction( ChooseFilesToSave, FExecuteAction::CreateStatic( &FMainFrameActionCallbacks::ChoosePackagesToSave ), FCanExecuteAction::CreateStatic( &FMainFrameActionCallbacks::CanSaveWorld ) );
 
-	UI_COMMAND( ChooseFilesToCheckIn, "Submit to Source Control...", "Opens a dialog with check in options for content and levels", EUserInterfaceActionType::Button, FInputChord() );
-	ActionList->MapAction( ChooseFilesToCheckIn, FExecuteAction::CreateStatic( &FMainFrameActionCallbacks::ChoosePackagesToCheckIn ), FCanExecuteAction::CreateStatic( &FMainFrameActionCallbacks::CanChoosePackagesToCheckIn ) );
+	UI_COMMAND( ViewChangelists, "View Changelists...", "Opens a dialog with check in options for content and levels", EUserInterfaceActionType::Button, FInputChord() );
+	ActionList->MapAction(ViewChangelists, FExecuteAction::CreateStatic( &FMainFrameActionCallbacks::ViewChangelists ), FCanExecuteAction::CreateStatic( &FMainFrameActionCallbacks::CanViewChangelists ) );
 
 	UI_COMMAND( ConnectToSourceControl, "Connect To Source Control...", "Connect to source control to allow source control operations to be performed on content and levels.", EUserInterfaceActionType::Button, FInputChord() );
 	ActionList->MapAction( ConnectToSourceControl, FExecuteAction::CreateStatic( &FMainFrameActionCallbacks::ConnectToSourceControl ), DefaultExecuteAction );
@@ -266,14 +266,14 @@ void FMainFrameActionCallbacks::ChoosePackagesToSave()
 	FEditorFileUtils::SaveDirtyPackages( bPromptUserToSave, bSaveMapPackages, bSaveContentPackages, bFastSave, bNotifyNoPackagesSaved, bCanBeDeclined );
 }
 
-void FMainFrameActionCallbacks::ChoosePackagesToCheckIn()
+void FMainFrameActionCallbacks::ViewChangelists()
 {
-	FSourceControlWindows::ChoosePackagesToCheckIn();
+	ISourceControlWindowsModule::Get().ShowChangelistsTab();
 }
 
-bool FMainFrameActionCallbacks::CanChoosePackagesToCheckIn()
+bool FMainFrameActionCallbacks::CanViewChangelists()
 {
-	return FSourceControlWindows::CanChoosePackagesToCheckIn();
+	return ISourceControlWindowsModule::Get().CanShowChangelistsTab();
 }
 
 void FMainFrameActionCallbacks::ConnectToSourceControl()
diff --git a/Engine/Source/Editor/MainFrame/Private/Frame/MainFrameActions.h b/Engine/Source/Editor/MainFrame/Private/Frame/MainFrameActions.h
--- a/Engine/Source/Editor/MainFrame/Private/Frame/MainFrameActions.h
+++ b/Engine/Source/Editor/MainFrame/Private/Frame/MainFrameActions.h
@@ -26,7 +26,7 @@ public:
 	TSharedPtr< FUICommandInfo > SaveAll;
 	TSharedPtr< FUICommandInfo > Exit;
 	TSharedPtr< FUICommandInfo > ChooseFilesToSave;
-	TSharedPtr< FUICommandInfo > ChooseFilesToCheckIn;
+	TSharedPtr< FUICommandInfo > ViewChangelists;
 	TSharedPtr< FUICommandInfo > ConnectToSourceControl;
 	TSharedPtr< FUICommandInfo > NewProject;
 	TSharedPtr< FUICommandInfo > OpenProject;
@@ -106,11 +106,11 @@ public:
 	/** Opens a dialog to choose packages to save */
 	static void ChoosePackagesToSave();
 
-	/** Opens a dialog to choose packages to submit */
-	static void ChoosePackagesToCheckIn();
+	/** Opens a dialog to view the pending changelists */
+	static void ViewChangelists();
 
-	/** Determines whether we can choose packages to check in (we cant if an operation is already in progress) */
-	static bool CanChoosePackagesToCheckIn();
+	/** Determines whether we can show the changelist view */
+	static bool CanViewChangelists();
 
 	/** Enable source control features */
 	static void ConnectToSourceControl();
diff --git a/Engine/Source/Editor/MainFrame/Private/Menus/MainMenu.cpp b/Engine/Source/Editor/MainFrame/Private/Menus/MainMenu.cpp
--- a/Engine/Source/Editor/MainFrame/Private/Menus/MainMenu.cpp
+++ b/Engine/Source/Editor/MainFrame/Private/Menus/MainMenu.cpp
@@ -56,11 +56,17 @@ void FMainMenu::RegisterFileMenu()
 			if (ISourceControlModule::Get().IsEnabled() && ISourceControlModule::Get().GetProvider().IsAvailable())
 			{
 				// Choose specific files to submit
-				InSection.AddMenuEntry(FMainFrameCommands::Get().ChooseFilesToCheckIn);
+				InSection.AddMenuEntry(FMainFrameCommands::Get().ViewChangelists, 
+					TAttribute<FText>(),
+					TAttribute<FText>(),
+					FSlateIcon(FEditorStyle::GetStyleSetName(), "SourceControl.ChangelistsTab"));
 			}
 			else
 			{
-				InSection.AddMenuEntry(FMainFrameCommands::Get().ConnectToSourceControl);
+				InSection.AddMenuEntry(FMainFrameCommands::Get().ConnectToSourceControl, 
+					TAttribute<FText>(),
+					TAttribute<FText>(),
+					FSlateIcon(FEditorStyle::GetStyleSetName(), "SourceControl.Actions.Connect"));
 			}
 		}));
 	}
diff --git a/Engine/Source/Editor/SourceControlWindows/Private/SourceControlWindowsModule.cpp b/Engine/Source/Editor/SourceControlWindows/Private/SourceControlWindowsModule.cpp
--- a/Engine/Source/Editor/SourceControlWindows/Private/SourceControlWindowsModule.cpp
+++ b/Engine/Source/Editor/SourceControlWindows/Private/SourceControlWindowsModule.cpp
@@ -3,6 +3,7 @@
 #include "CoreMinimal.h"
 #include "Modules/ModuleManager.h"
 #include "ISourceControlWindowsModule.h"
+#include "ISourceControlModule.h"
 
 #include "Widgets/Docking/SDockTab.h"
 #include "Textures/SlateIcon.h"
@@ -34,6 +35,7 @@ public:
 	virtual void ShutdownModule() override;
 
 	virtual void ShowChangelistsTab() override;
+	virtual bool CanShowChangelistsTab() const override;
 
 private:
 	TSharedRef<SDockTab> CreateChangelistsTab(const FSpawnTabArgs& Args);
@@ -119,4 +121,16 @@ void FSourceControlWindowsModule::ShowChangelistsTab()
 	FGlobalTabmanager::Get()->TryInvokeTab(FTabId(SourceControlChangelistsTabName));
 }
 
+bool FSourceControlWindowsModule::CanShowChangelistsTab() const
+{
+	ISourceControlModule& SourceControlModule = ISourceControlModule::Get();
+	if (ISourceControlModule::Get().IsEnabled() &&
+		ISourceControlModule::Get().GetProvider().IsAvailable())
+	{
+		return true;
+	}
+
+	return false;
+}
+
 #undef LOCTEXT_NAMESPACE
diff --git a/Engine/Source/Editor/SourceControlWindows/Public/ISourceControlWindowsModule.h b/Engine/Source/Editor/SourceControlWindows/Public/ISourceControlWindowsModule.h
--- a/Engine/Source/Editor/SourceControlWindows/Public/ISourceControlWindowsModule.h
+++ b/Engine/Source/Editor/SourceControlWindows/Public/ISourceControlWindowsModule.h
@@ -20,4 +20,5 @@ public:
 	}
 
 	virtual void ShowChangelistsTab() = 0;
+	virtual bool CanShowChangelistsTab() const = 0;
 };
diff --git a/Engine/Source/Editor/StatusBar/Private/SourceControlMenuHelpers.cpp b/Engine/Source/Editor/StatusBar/Private/SourceControlMenuHelpers.cpp
--- a/Engine/Source/Editor/StatusBar/Private/SourceControlMenuHelpers.cpp
+++ b/Engine/Source/Editor/StatusBar/Private/SourceControlMenuHelpers.cpp
@@ -38,7 +38,6 @@ void FSourceControlCommands::RegisterCommands()
 	UI_COMMAND(ChangeSourceControlSettings, "Change Source Control Settings...", "Opens a dialog to change source control settings.", EUserInterfaceActionType::Button, FInputChord());
 	UI_COMMAND(ViewChangelists, "View Changelists...", "Open the Source Control Changelist tab.", EUserInterfaceActionType::Button, FInputChord());
 	UI_COMMAND(CheckOutModifiedFiles, "Check Out Modified Files...", "Opens a dialog to check out any assets which have been modified.", EUserInterfaceActionType::Button, FInputChord());
-	UI_COMMAND(SubmitToSourceControl, "Submit to Source Control...", "Opens a dialog with check in options for content and levels.",   EUserInterfaceActionType::Button, FInputChord());
 
 	ActionList->MapAction(
 		ConnectToSourceControl,
@@ -61,13 +60,6 @@ void FSourceControlCommands::RegisterCommands()
 		FExecuteAction::CreateStatic(&FSourceControlCommands::CheckOutModifiedFiles_Clicked),
 		FCanExecuteAction::CreateStatic(&FSourceControlCommands::CheckOutModifiedFiles_CanExecute)
 	);
-
-	ActionList->MapAction(
-		SubmitToSourceControl,
-		FExecuteAction::CreateStatic(&FSourceControlCommands::SubmitToSourceControl_Clicked),
-		FCanExecuteAction::CreateStatic(&FSourceControlCommands::SubmitToSourceControl_CanExecute)
-	);
-
 }
 
 void FSourceControlCommands::ConnectToSourceControl_Clicked()
@@ -79,14 +71,7 @@ void FSourceControlCommands::ConnectToSourceControl_Clicked()
 
 bool FSourceControlCommands::ViewChangelists_CanExecute()
 {
-	ISourceControlModule& SourceControlModule = ISourceControlModule::Get();
-	if (ISourceControlModule::Get().IsEnabled() &&
-		ISourceControlModule::Get().GetProvider().IsAvailable())
-	{
-		return true;
-	}
-
-	return false;
+	return ISourceControlWindowsModule::Get().CanShowChangelistsTab();
 }
 
 void FSourceControlCommands::ViewChangelists_Clicked()
@@ -121,20 +106,6 @@ void FSourceControlCommands::CheckOutModifiedFiles_Clicked()
 	FEditorFileUtils::PromptForCheckoutAndSave(PackagesToSave, bCheckDirty, bPromptUserToSave);
 }
 
-bool FSourceControlCommands::SubmitToSourceControl_CanExecute()
-{
-	ISourceControlModule& SourceControlModule = ISourceControlModule::Get();
-	return ISourceControlModule::Get().IsEnabled() &&
-		ISourceControlModule::Get().GetProvider().IsAvailable() &&
-		FSourceControlWindows::CanChoosePackagesToCheckIn();
-}
-
-void FSourceControlCommands::SubmitToSourceControl_Clicked()
-{
-	FSourceControlWindows::ChoosePackagesToCheckIn();
-}
-
-
 FSourceControlMenuHelpers::EQueryState FSourceControlMenuHelpers::QueryState = FSourceControlMenuHelpers::EQueryState::NotQueried;
 
 void FSourceControlMenuHelpers::CheckSourceControlStatus()
@@ -199,14 +170,6 @@ TSharedRef<SWidget> FSourceControlMenuHelpers::GenerateSourceControlMenuContent(
 		FSlateIcon(FEditorStyle::GetStyleSetName(), "SourceControl.Actions.CheckOut")
 	);
 
-
-	Section.AddMenuEntry(
-		FSourceControlCommands::Get().SubmitToSourceControl,
-		TAttribute<FText>(),
-		TAttribute<FText>(),
-		FSlateIcon(FEditorStyle::GetStyleSetName(), "SourceControl.Actions.Submit")
-	);
-
 	return UToolMenus::Get()->GenerateWidget("StatusBar.ToolBar.SourceControl", FToolMenuContext(FSourceControlCommands::ActionList));
 }
 
diff --git a/Engine/Source/Editor/StatusBar/Private/SourceControlMenuHelpers.h b/Engine/Source/Editor/StatusBar/Private/SourceControlMenuHelpers.h
--- a/Engine/Source/Editor/StatusBar/Private/SourceControlMenuHelpers.h
+++ b/Engine/Source/Editor/StatusBar/Private/SourceControlMenuHelpers.h
@@ -26,8 +26,6 @@ private:
 	static void ViewChangelists_Clicked();
 	static bool CheckOutModifiedFiles_CanExecute();
 	static void CheckOutModifiedFiles_Clicked();
-	static bool SubmitToSourceControl_CanExecute();
-	static void SubmitToSourceControl_Clicked();
 
 public:
 	/**
@@ -37,7 +35,6 @@ public:
 	TSharedPtr< FUICommandInfo > ChangeSourceControlSettings;
 	TSharedPtr< FUICommandInfo > ViewChangelists;
 	TSharedPtr< FUICommandInfo > CheckOutModifiedFiles;
-	TSharedPtr< FUICommandInfo > SubmitToSourceControl;
 
 	static TSharedRef<FUICommandList> ActionList;
 };
