Fix progress notifications not appearing if no editor window has focus

diff --git a/Engine/Source/Editor/StatusBar/Private/StatusBarSubsystem.cpp b/Engine/Source/Editor/StatusBar/Private/StatusBarSubsystem.cpp
--- a/Engine/Source/Editor/StatusBar/Private/StatusBarSubsystem.cpp
+++ b/Engine/Source/Editor/StatusBar/Private/StatusBarSubsystem.cpp
@@ -176,13 +176,21 @@ void UStatusBarSubsystem::ClearStatusBarMessages(FName StatusBarName)
 
 void UStatusBarSubsystem::StartProgressNotification(FProgressNotificationHandle Handle, FText DisplayText, int32 TotalWorkToDo)
 {
+	// Get the active window, if one is not active a notification was started when the application was deactivated so use the focus path to find a window or just use the root window if there is no keyboard focus
+	TSharedPtr<SWindow> ActiveWindow = FSlateApplication::Get().GetActiveTopLevelRegularWindow();
+	if (!ActiveWindow)
+	{
+		TSharedPtr<SWidget> FocusedWidget = FSlateApplication::Get().GetKeyboardFocusedWidget();
+		ActiveWindow = FocusedWidget ? FSlateApplication::Get().FindWidgetWindow(FocusedWidget.ToSharedRef()) : FGlobalTabmanager::Get()->GetRootWindow();
+	}
+
 	// Find the active status bar to display the progress in
 	for (auto StatusBar : StatusBars)
 	{
 		if (TSharedPtr<SStatusBar> StatusBarPinned = StatusBar.Value.Pin())
 		{
 			TSharedPtr<SDockTab> ParentTab = StatusBarPinned->GetParentTab();
-			if (ParentTab && ParentTab->IsForeground() && ParentTab->GetParentWindow() == FSlateApplication::Get().GetActiveTopLevelRegularWindow())
+			if (ParentTab && ParentTab->IsForeground() && ParentTab->GetParentWindow() == ActiveWindow)
 			{
 				StatusBarPinned->StartProgressNotification(Handle, DisplayText, TotalWorkToDo);
 				break;
