UE5 layout changes

Content browser now hidden by default, opening tip is displayed showing you how to access it from the drawer. It can be permanently docked into the layout from there. Opening tip is never shown again
Placement browser is now closed by default
A UE4 classic layout has been added for those who want the placement browser and content browser docked by default

diff --git a/Engine/Source/Editor/ContentBrowser/Private/ContentBrowserSingleton.cpp b/Engine/Source/Editor/ContentBrowser/Private/ContentBrowserSingleton.cpp
--- a/Engine/Source/Editor/ContentBrowser/Private/ContentBrowserSingleton.cpp
+++ b/Engine/Source/Editor/ContentBrowser/Private/ContentBrowserSingleton.cpp
@@ -33,6 +33,8 @@
 
 #define LOCTEXT_NAMESPACE "ContentBrowser"
 
+static const FName ContentBrowserDrawerInstanceName("ContentBrowserDrawer");
+
 FContentBrowserSingleton::FContentBrowserSingleton()
 	: CollectionAssetRegistryBridge(MakeShared<FCollectionAssetRegistryBridge>())
 	, SettingsStringID(0)
@@ -87,10 +89,6 @@ FContentBrowserSingleton::FContentBrowserSingleton()
 		AddSpawnerEntryToMenuSection(Section, ContentBrowserTabSpawner, TabID);
 	}
 
-	// Register a couple legacy tab ids
-	FGlobalTabmanager::Get()->AddLegacyTabType( "LevelEditorContentBrowser", "ContentBrowserTab1" );
-	FGlobalTabmanager::Get()->AddLegacyTabType( "MajorContentBrowserTab", "ContentBrowserTab2" );
-
 	// Register to be notified when properties are edited
 	FEditorDelegates::LoadSelectedAssetsIfNeeded.AddRaw(this, &FContentBrowserSingleton::OnEditorLoadSelectedAssetsIfNeeded);
 
@@ -133,16 +131,17 @@ TSharedRef<class SWidget> FContentBrowserSingleton::CreateCollectionPicker(const
 		.CollectionPickerConfig(CollectionPickerConfig);
 }
 
-TSharedRef<class SWidget> FContentBrowserSingleton::CreateContentBrowserDrawer(const FContentBrowserConfig& ContentBrowserConfig)
+TSharedRef<class SWidget> FContentBrowserSingleton::CreateContentBrowserDrawer(const FContentBrowserConfig& ContentBrowserConfig, TFunction<TSharedPtr<SDockTab>()> InOnGetTabForDrawer)
 {
 	TSharedPtr<SContentBrowser> ContentBrowserDrawerPinned;
 	if(!ContentBrowserDrawer.IsValid())
 	{
-		static const FName ContentBrowserDrawerInstanceName("ContentBrowserDrawer");
 		ContentBrowserDrawerPinned =
 			SNew(SContentBrowser, ContentBrowserDrawerInstanceName, &ContentBrowserConfig)
-			.IsEnabled(FSlateApplication::Get().GetNormalExecutionAttribute());
+			.IsEnabled(FSlateApplication::Get().GetNormalExecutionAttribute())
+			.IsDrawer(true);
 
+		OnGetTabForDrawer = InOnGetTabForDrawer;
 		ContentBrowserDrawer = ContentBrowserDrawerPinned;
 	}
 	else
@@ -656,7 +655,55 @@ void FContentBrowserSingleton::FocusContentBrowser(const TSharedPtr<SContentBrow
 	}
 }
 
-FName FContentBrowserSingleton::SummonNewBrowser(bool bAllowLockedBrowsers)
+void FContentBrowserSingleton::DockContentBrowserDrawer()
+{
+	TSharedPtr<FTabManager> ForTabManager;
+	TSharedPtr<SDockTab> Tab = OnGetTabForDrawer();
+	if (Tab)
+	{
+		ForTabManager = FGlobalTabmanager::Get()->GetTabManagerForMajorTab(Tab.ToSharedRef());
+	}
+
+	
+	// Dont summon a content browser if a content browser already exists in the tab manager
+	bool bExistingTab = false;
+	{
+		for (int32 BrowserIdx = 0; BrowserIdx < UE_ARRAY_COUNT(ContentBrowserTabIDs); BrowserIdx++)
+		{
+			if (TSharedPtr<SDockTab> ExistingTab = ForTabManager->FindExistingLiveTab(ContentBrowserTabIDs[BrowserIdx]))
+			{
+				GEditor->GetEditorSubsystem<UStatusBarSubsystem>()->ForceDismissDrawer();
+				ExistingTab->ActivateInParent(ETabActivationCause::SetDirectly);
+				bExistingTab = true;
+				break;
+			}
+		}
+	}
+
+	if(!bExistingTab)
+	{
+		TSharedPtr<SContentBrowser> ContentBrowserDrawerPinned = ContentBrowserDrawer.Pin();
+
+		// Make sure current content browser drawer settings are saved so we can copy them to the new browser
+		ContentBrowserDrawerPinned->SaveSettings();
+
+		FName InstanceName = SummonNewBrowser(false, ForTabManager);
+		for (int32 AllBrowsersIdx = AllContentBrowsers.Num() - 1; AllBrowsersIdx >= 0; --AllBrowsersIdx)
+		{
+			if (TSharedPtr<SContentBrowser> Browser = AllContentBrowsers[AllBrowsersIdx].Pin())
+			{
+				if (Browser->GetInstanceName() == InstanceName)
+				{
+					Browser->CopySettingsFromBrowser(ContentBrowserDrawerPinned);
+					break;
+				}
+			}
+		}
+	}
+
+}
+
+FName FContentBrowserSingleton::SummonNewBrowser(bool bAllowLockedBrowsers, TSharedPtr<FTabManager> SpecificTabManager)
 {
 	TSet<FName> OpenBrowserIDs;
 
@@ -684,7 +731,7 @@ FName FContentBrowserSingleton::SummonNewBrowser(bool bAllowLockedBrowsers)
 
 	if ( NewTabName != NAME_None )
 	{
-		const TWeakPtr<FTabManager>& TabManagerToInvoke = BrowserToLastKnownTabManagerMap.FindRef(NewTabName);
+		const TWeakPtr<FTabManager>& TabManagerToInvoke = SpecificTabManager.IsValid() ? SpecificTabManager : BrowserToLastKnownTabManagerMap.FindRef(NewTabName);
 		if ( TabManagerToInvoke.IsValid() )
 		{
 			TabManagerToInvoke.Pin()->TryInvokeTab(NewTabName);
diff --git a/Engine/Source/Editor/ContentBrowser/Private/ContentBrowserSingleton.h b/Engine/Source/Editor/ContentBrowser/Private/ContentBrowserSingleton.h
--- a/Engine/Source/Editor/ContentBrowser/Private/ContentBrowserSingleton.h
+++ b/Engine/Source/Editor/ContentBrowser/Private/ContentBrowserSingleton.h
@@ -49,7 +49,7 @@ public:
 	virtual TSharedRef<class SWidget> CreateAssetPicker(const FAssetPickerConfig& AssetPickerConfig) override;
 	virtual TSharedRef<class SWidget> CreatePathPicker(const FPathPickerConfig& PathPickerConfig) override;
 	virtual TSharedRef<class SWidget> CreateCollectionPicker(const FCollectionPickerConfig& CollectionPickerConfig) override;
-	virtual TSharedRef<class SWidget> CreateContentBrowserDrawer(const FContentBrowserConfig& ContentBrowserConfig) override;
+	virtual TSharedRef<class SWidget> CreateContentBrowserDrawer(const FContentBrowserConfig& ContentBrowserConfig, TFunction<TSharedPtr<SDockTab>()> InOnGetTabForDrawer) override;
 	virtual void CreateOpenAssetDialog(const FOpenAssetDialogConfig& OpenAssetConfig, const FOnAssetsChosenForOpen& OnAssetsChosenForOpen, const FOnAssetDialogCancelled& OnAssetDialogCancelled) override;
 	virtual TArray<FAssetData> CreateModalOpenAssetDialog(const FOpenAssetDialogConfig& InConfig) override;
 	virtual void CreateSaveAssetDialog(const FSaveAssetDialogConfig& SaveAssetConfig, const FOnObjectPathChosenForSave& OnAssetNameChosenForSave, const FOnAssetDialogCancelled& OnAssetDialogCancelled) override;
@@ -90,6 +90,9 @@ public:
 	/** Single storage location for content browser favorites */
 	TArray<FString> FavoriteFolderPaths;
 
+	/** Docks the current content browser drawer as a tabbed content browser in a layout */
+	void DockContentBrowserDrawer();
+
 private:
 
 	/** Util to get or create the content browser that should be used by the various Sync functions */
@@ -110,7 +113,7 @@ private:
 	void FocusContentBrowser(const TSharedPtr<SContentBrowser>& BrowserToFocus);
 
 	/** Summons a new content browser */
-	FName SummonNewBrowser(bool bAllowLockedBrowsers = false);
+	FName SummonNewBrowser(bool bAllowLockedBrowsers = false, TSharedPtr<FTabManager> SpecificTabManager = nullptr);
 
 	/** Handler for a request to spawn a new content browser tab */
 	TSharedRef<SDockTab> SpawnContentBrowserTab( const FSpawnTabArgs& SpawnTabArgs, int32 BrowserIdx );
@@ -136,6 +139,8 @@ private:
 
 	TWeakPtr<SContentBrowser> ContentBrowserDrawer;
 
+	TFunction<TSharedPtr<SDockTab>()> OnGetTabForDrawer;
+
 	TMap<FName, TWeakPtr<FTabManager>> BrowserToLastKnownTabManagerMap;
 
 	TWeakPtr<SContentBrowser> PrimaryContentBrowser;
diff --git a/Engine/Source/Editor/ContentBrowser/Private/SContentBrowser.cpp b/Engine/Source/Editor/ContentBrowser/Private/SContentBrowser.cpp
--- a/Engine/Source/Editor/ContentBrowser/Private/SContentBrowser.cpp
+++ b/Engine/Source/Editor/ContentBrowser/Private/SContentBrowser.cpp
@@ -127,8 +127,10 @@ void SContentBrowser::Construct( const FArguments& InArgs, const FName& InInstan
 		InArgs._ContainingTab->SetOnTabActivated( SDockTab::FOnTabActivatedCallback::CreateSP( this, &SContentBrowser::OnContainingTabActivated ) );
 	}
 	
-	bIsLocked = InArgs._InitiallyLocked;
+
+ 	bIsLocked = InArgs._InitiallyLocked;
 	bCanSetAsPrimaryBrowser = Config != nullptr ? Config->bCanSetAsPrimaryBrowser : true;
+	bIsDrawer = InArgs._IsDrawer;
 
 	HistoryManager.SetOnApplyHistoryData(FOnApplyHistoryData::CreateSP(this, &SContentBrowser::OnApplyHistoryData));
 	HistoryManager.SetOnUpdateHistoryData(FOnUpdateHistoryData::CreateSP(this, &SContentBrowser::OnUpdateHistoryData));
@@ -236,7 +238,7 @@ void SContentBrowser::Construct( const FArguments& InArgs, const FName& InInstan
 		[
 			SNew( SBorder )
 			.Padding( FMargin( 3 ) )
-			.BorderImage(FAppStyle::Get().GetBrush("Brushes.Panel"))
+			.BorderImage(bIsDrawer ? FStyleDefaults::GetNoBrush() : FAppStyle::Get().GetBrush("Brushes.Panel"))
 			[
 				SNew(SHorizontalBox)
 				+ SHorizontalBox::Slot()
@@ -333,6 +335,13 @@ void SContentBrowser::Construct( const FArguments& InArgs, const FName& InInstan
 				]
 				+ SHorizontalBox::Slot()
 				.AutoWidth()
+				.HAlign(HAlign_Right)
+				.VAlign(VAlign_Center)
+				[
+					CreateDrawerDockButton(Config)
+				]
+				+ SHorizontalBox::Slot()
+				.AutoWidth()
 				.Padding(5.0f, 0.0f, 0.0f, 0.0f)
 				.HAlign(HAlign_Right)
 				[
@@ -378,7 +387,7 @@ void SContentBrowser::Construct( const FArguments& InArgs, const FName& InInstan
 						[
 							SAssignNew(PathFavoriteSplitterPtr, SSplitter)
 							.Clipping(EWidgetClipping::ClipToBounds)
-							.PhysicalSplitterHandleSize(3.0f)
+							.PhysicalSplitterHandleSize(1.0f)
 							.HitDetectionSplitterHandleSize(3.0f)
 							.Orientation(EOrientation::Orient_Vertical)
 							.MinimumSlotHeight(26.0f)
@@ -867,6 +876,41 @@ TSharedRef<SWidget> SContentBrowser::CreateDockedCollectionsView(const FContentB
 		];
 }
 
+TSharedRef<SWidget> SContentBrowser::CreateDrawerDockButton(const FContentBrowserConfig* Config)
+{
+	if(bIsDrawer)
+	{
+		return 
+			SNew(SButton)
+			.ButtonStyle(FAppStyle::Get(), "SimpleButton")
+			.ToolTipText(LOCTEXT("DockInLayout_Tooltip", "Docks this content browser in the current layout, copying all settings from the drawer.\nThe drawer will still be usable as a temporary browser."))
+			.ContentPadding(FMargin(1, 0))
+			.OnClicked(this, &SContentBrowser::DockInLayoutClicked)
+			[
+				SNew(SHorizontalBox)
+				+ SHorizontalBox::Slot()
+				.AutoWidth()
+				.VAlign(VAlign_Center)
+				.Padding(4.0, 0.0f)
+				[
+					SNew(SImage)
+					.ColorAndOpacity(FSlateColor::UseForeground())
+					.Image(FAppStyle::Get().GetBrush("EditorViewport.SubMenu.Layouts"))
+				]
+				+ SHorizontalBox::Slot()
+				.VAlign(VAlign_Center)
+				.Padding(4.0, 0.0f)
+				[
+					SNew(STextBlock)
+					.Text(LOCTEXT("DockInLayout", "Dock in Layout"))
+					.ColorAndOpacity(FSlateColor::UseForeground())
+				]
+			];
+	}
+
+	return SNullWidget::NullWidget;
+}
+
 void SContentBrowser::ExtendViewOptionsMenu(const FContentBrowserConfig* Config)
 {
 	UToolMenu* Menu = UToolMenus::Get()->ExtendMenu("ContentBrowser.AssetViewOptions");
@@ -1342,6 +1386,16 @@ void SContentBrowser::SetKeyboardFocusOnSearch() const
 	FSlateApplication::Get().SetKeyboardFocus( SearchBoxPtr, EFocusCause::SetDirectly );
 }
 
+void SContentBrowser::CopySettingsFromBrowser(TSharedPtr<SContentBrowser> OtherBrowser)
+{
+	FName InstanceNameToCopyFrom = OtherBrowser->GetInstanceName();
+
+	// Clear out any existing settings that dont get reset on load
+	FilterListPtr->RemoveAllFilters();
+
+	LoadSettings(InstanceNameToCopyFrom);
+}
+
 FReply SContentBrowser::OnKeyDown( const FGeometry& MyGeometry, const FKeyEvent& InKeyEvent )
 {
 	if( Commands->ProcessCommandBindings( InKeyEvent ) )
@@ -2426,6 +2480,13 @@ FReply SContentBrowser::ToggleLockClicked()
 	return FReply::Handled();
 }
 
+FReply SContentBrowser::DockInLayoutClicked()
+{
+	FContentBrowserSingleton::Get().DockContentBrowserDrawer();
+
+	return FReply::Handled();
+}
+
 FText SContentBrowser::GetLockMenuText() const
 {
 	return IsLocked() ? LOCTEXT("ContentBrowserLockMenu_Unlock", "Unlock Content Browser") : LOCTEXT("ContentBrowserLockMenu_Lock", "Lock Content Browser");
diff --git a/Engine/Source/Editor/ContentBrowser/Private/SContentBrowser.h b/Engine/Source/Editor/ContentBrowser/Private/SContentBrowser.h
--- a/Engine/Source/Editor/ContentBrowser/Private/SContentBrowser.h
+++ b/Engine/Source/Editor/ContentBrowser/Private/SContentBrowser.h
@@ -54,13 +54,15 @@ public:
 	SLATE_BEGIN_ARGS( SContentBrowser )
 		: _ContainingTab()
 		, _InitiallyLocked(false)
+		, _IsDrawer(false)
 	{}
 		/** The tab in which the content browser resides */
-		SLATE_ARGUMENT( TSharedPtr<SDockTab>, ContainingTab )
+		SLATE_ARGUMENT(TSharedPtr<SDockTab>, ContainingTab)
 
 		/** If true, this content browser will not sync from external sources. */
-		SLATE_ARGUMENT( bool, InitiallyLocked )
+		SLATE_ARGUMENT(bool, InitiallyLocked )
 
+		SLATE_ARGUMENT(bool, IsDrawer)
 	SLATE_END_ARGS()
 
 	~SContentBrowser();
@@ -178,6 +180,13 @@ public:
 	/** Gives keyboard focus to the asset search box */
 	void SetKeyboardFocusOnSearch() const;
 
+	/**
+	 * Copies settings from a different browser to this browser
+	 * Note this overrides any settings already saved for this browser 
+	 */
+	void CopySettingsFromBrowser(TSharedPtr<SContentBrowser> OtherBrowser);
+
+	/** SWidget interface  */
 	virtual FReply OnKeyDown( const FGeometry& MyGeometry, const FKeyEvent& InKeyEvent ) override;
 	virtual FReply OnMouseButtonDown( const FGeometry& MyGeometry, const FPointerEvent& MouseEvent ) override;
 	virtual FReply OnPreviewMouseButtonDown( const FGeometry& MyGeometry, const FPointerEvent& MouseEvent ) override;
@@ -322,6 +331,9 @@ private:
 	/** Handler for clicking the lock button */
 	FReply ToggleLockClicked();
 
+	/** Handler for clicking the dock drawer in layout button */
+	FReply DockInLayoutClicked();
+
 	/** Gets the menu text */
 	FText GetLockMenuText() const;
 
@@ -515,6 +527,7 @@ private:
 	TSharedRef<SWidget> CreateFavoritesView(const FContentBrowserConfig* Config);
 	TSharedRef<SWidget> CreatePathView(const FContentBrowserConfig* Config);
 	TSharedRef<SWidget> CreateDockedCollectionsView(const FContentBrowserConfig* Config);
+	TSharedRef<SWidget> CreateDrawerDockButton(const FContentBrowserConfig* Config);
 
 	/** Adds menu options to the view menu */
 	void ExtendViewOptionsMenu(const FContentBrowserConfig* Config);
@@ -599,6 +612,9 @@ private:
 	/** True if this content browser can be set to the primary browser. */
 	bool bCanSetAsPrimaryBrowser;
 
+	/** True if this content browser is an a drawer */
+	bool bIsDrawer;
+
 	/** Unique name for this Content Browser. */
 	FName InstanceName;
 
diff --git a/Engine/Source/Editor/ContentBrowser/Public/IContentBrowserSingleton.h b/Engine/Source/Editor/ContentBrowser/Public/IContentBrowserSingleton.h
--- a/Engine/Source/Editor/ContentBrowser/Public/IContentBrowserSingleton.h
+++ b/Engine/Source/Editor/ContentBrowser/Public/IContentBrowserSingleton.h
@@ -507,7 +507,7 @@ public:
 	 *
 	 * @return The content browser drawer widget
 	 */
-	virtual TSharedRef<class SWidget> CreateContentBrowserDrawer(const FContentBrowserConfig& ContentBrowserConfig) = 0;
+	virtual TSharedRef<class SWidget> CreateContentBrowserDrawer(const FContentBrowserConfig& ContentBrowserConfig, TFunction<TSharedPtr<SDockTab>()> InOnGetTabForDrawer) = 0;
 
 	/**
 	 * Opens the Open Asset dialog in a non-modal window
diff --git a/Engine/Source/Editor/EditorStyle/Private/StarshipStyle.cpp b/Engine/Source/Editor/EditorStyle/Private/StarshipStyle.cpp
--- a/Engine/Source/Editor/EditorStyle/Private/StarshipStyle.cpp
+++ b/Engine/Source/Editor/EditorStyle/Private/StarshipStyle.cpp
@@ -7852,7 +7852,8 @@ void FStarshipEditorStyle::FStyle::SetupStatusBarStyle()
 
 	Set("StatusBar.ContentBrowserUp", new CORE_IMAGE_BRUSH_SVG("Starship/Common/chevron-up", Icon16x16));
 	Set("StatusBar.ContentBrowserDown", new CORE_IMAGE_BRUSH_SVG("Starship/Common/chevron-down", Icon16x16));
-	Set("StatusBar.ContentBrowserShadow", new BOX_BRUSH("Starship/StatusBar/drawer-shadow-bottom", FMargin(10/64.0f, 20/64.f, 10/64.f, 0), FLinearColor(0,0,0,1)));
+	Set("StatusBar.DrawerShadow", new BOX_BRUSH("Starship/StatusBar/drawer-shadow-bottom", FMargin(10/64.0f, 20/64.f, 10/64.f, 0), FLinearColor(0,0,0,1)));
+	Set("StatusBar.DrawerBackground", new FSlateRoundedBoxBrush(FStyleColors::Panel, 5.0f));
 
 	Set("StatusBar.Height", 32.0f);
 	Set("StatusBar.Background", new FSlateColorBrush(FStyleColors::Panel));
diff --git a/Engine/Source/Editor/LevelEditor/Private/SLevelEditor.cpp b/Engine/Source/Editor/LevelEditor/Private/SLevelEditor.cpp
--- a/Engine/Source/Editor/LevelEditor/Private/SLevelEditor.cpp
+++ b/Engine/Source/Editor/LevelEditor/Private/SLevelEditor.cpp
@@ -1434,7 +1434,7 @@ TSharedRef<SWidget> SLevelEditor::RestoreContentArea( const TSharedRef<SDockTab>
 	//     9.3. Etc
 	// 10. Push the new "DefaultLayout.ini" together with your new code.
 	// 11. Also update these instructions if you change the version number (e.g., from "UnrealEd_Layout_v1.6" to "UnrealEd_Layout_v1.7").
-	const FName LayoutName = TEXT("LevelEditor_Layout_v1.6");
+	const FName LayoutName = TEXT("LevelEditor_Layout_v1.7");
 	const TSharedRef<FTabManager::FLayout> DefaultLayout = FLayoutSaveRestore::LoadFromConfig(GEditorLayoutIni,
 		FTabManager::NewLayout( LayoutName )
 		->AddArea
@@ -1463,7 +1463,7 @@ TSharedRef<SWidget> SLevelEditor::RestoreContentArea( const TSharedRef<SDockTab>
 					(
 						FTabManager::NewStack()
 						->SetSizeCoefficient( 0.3f )
-						->AddTab(LevelEditorTabIds::PlacementBrowser, ETabState::SidebarTab, ESidebarLocation::Left, 0.12f)
+						->AddTab(LevelEditorTabIds::PlacementBrowser, ETabState::ClosedTab)
 						->AddTab(LevelEditorTabIds::LevelEditorToolBox, ETabState::ClosedTab)
 					)
 					->Split
@@ -1478,7 +1478,8 @@ TSharedRef<SWidget> SLevelEditor::RestoreContentArea( const TSharedRef<SDockTab>
 				(
 					FTabManager::NewStack()
 					->SetSizeCoefficient(.4)
-					->AddTab("ContentBrowserTab1", ETabState::OpenedTab)
+					->SetHideTabWell(true)
+					->AddTab("ContentBrowserTab1", ETabState::ClosedTab)
 					->AddTab(LevelEditorTabIds::OutputLog, ETabState::ClosedTab)
 				)
 			)
diff --git a/Engine/Source/Editor/MainFrame/Private/MainFrameModule.cpp b/Engine/Source/Editor/MainFrame/Private/MainFrameModule.cpp
--- a/Engine/Source/Editor/MainFrame/Private/MainFrameModule.cpp
+++ b/Engine/Source/Editor/MainFrame/Private/MainFrameModule.cpp
@@ -195,7 +195,7 @@ void FMainFrameModule::CreateDefaultMainFrameAuxiliary(const bool bStartImmersiv
 			//     9.3. Etc
 			// 10. Push the new "DefaultLayout.ini" together with your new code.
 			// 11. Also update these instructions if you change the version number (e.g., from "UnrealEd_Layout_v1.4" to "UnrealEd_Layout_v1.5").
-			const FName LayoutName = TEXT("UnrealEd_Layout_v1.4");
+			const FName LayoutName = TEXT("UnrealEd_Layout_v1.5");
 			const TSharedRef<FTabManager::FLayout> DefaultLayout =
 				// We persist the positioning of the level editor and the content browser.
 				// The asset editors currently do not get saved.
@@ -213,17 +213,6 @@ void FMainFrameModule::CreateDefaultMainFrameAuxiliary(const bool bStartImmersiv
 					)
 				)
 				->AddArea
-				(
-					// content browser window
-					FTabManager::NewArea(WindowSize)
-					->Split
-					(
-						FTabManager::NewStack()
-						->SetSizeCoefficient(1.0f)
-						->AddTab("ContentBrowser1Tab", ETabState::ClosedTab)
-					)
-				)
-				->AddArea
 				(
 					// toolkits window
 					FTabManager::NewArea(WindowSize)
diff --git a/Engine/Source/Editor/StatusBar/Private/SStatusBar.cpp b/Engine/Source/Editor/StatusBar/Private/SStatusBar.cpp
--- a/Engine/Source/Editor/StatusBar/Private/SStatusBar.cpp
+++ b/Engine/Source/Editor/StatusBar/Private/SStatusBar.cpp
@@ -50,7 +50,7 @@ class SDrawerOverlay : public SCompoundWidget
 		_Clipping = EWidgetClipping::ClipToBounds;
 		_ShadowOffset = FVector2D(10.0f, 20.0f);
 	}
-	SLATE_DEFAULT_SLOT(FArguments, Content)
+		SLATE_DEFAULT_SLOT(FArguments, Content)
 		SLATE_ARGUMENT(float, MinDrawerHeight)
 		SLATE_ARGUMENT(float, MaxDrawerHeight)
 		SLATE_ARGUMENT(float, TargetDrawerHeight)
@@ -82,8 +82,9 @@ class SDrawerOverlay : public SCompoundWidget
 
 		OnTargetHeightChanged = InArgs._OnTargetHeightChanged;
 
-		BackgroundBrush = FAppStyle::Get().GetBrush("Brushes.Panel");
-		ShadowBrush = FAppStyle::Get().GetBrush("StatusBar.ContentBrowserShadow");
+		BackgroundBrush = FAppStyle::Get().GetBrush("StatusBar.DrawerBackground");
+		ShadowBrush = FAppStyle::Get().GetBrush("StatusBar.DrawerShadow");
+		BorderBrush = FAppStyle::Get().GetBrush("Docking.Sidebar.Border");
 
 		bIsResizeHandleHovered = false;
 		bIsResizing = false;
@@ -239,7 +240,16 @@ class SDrawerOverlay : public SCompoundWidget
 			ESlateDrawEffect::None,
 			ShadowBrush->GetTint(InWidgetStyle));
 
-		return OutLayerId;
+
+		FSlateDrawElement::MakeBox(
+			OutDrawElements,
+			OutLayerId+1,
+			RenderTransformedChildGeometry.ToPaintGeometry(ShadowOffset, FVector2D(AllottedGeometry.GetLocalSize().X - (ShadowOffset.X * 2), TargetHeight)),
+			BorderBrush,
+			ESlateDrawEffect::None,
+			BorderBrush->GetTint(InWidgetStyle));
+
+		return OutLayerId+1;
 
 	}
 
@@ -263,6 +273,7 @@ private:
 	FOnStatusBarDrawerTargetHeightChanged OnTargetHeightChanged;
 	const FSlateBrush* BackgroundBrush;
 	const FSlateBrush* ShadowBrush;
+	const FSlateBrush* BorderBrush;
 	const FSplitterStyle* SplitterStyle;
 	FVector2D ShadowOffset;
 	FThrottleRequest ResizeThrottleHandle;
diff --git a/Engine/Source/Editor/StatusBar/Private/StatusBarSubsystem.cpp b/Engine/Source/Editor/StatusBar/Private/StatusBarSubsystem.cpp
--- a/Engine/Source/Editor/StatusBar/Private/StatusBarSubsystem.cpp
+++ b/Engine/Source/Editor/StatusBar/Private/StatusBarSubsystem.cpp
@@ -14,15 +14,199 @@
 #include "Toolkits/GlobalEditorCommonCommands.h"
 #include "Widgets/Notifications/SNotificationList.h"
 #include "Framework/Notifications/NotificationManager.h"
+#include "Interfaces/IMainFrameModule.h"
+#include "Widgets/Notifications/SNotificationBackground.h"
+#include "Brushes/SlateRoundedBoxBrush.h"
+#include "Styling/StyleColors.h"
+#include "Widgets/Images/SImage.h"
+#include "Widgets/Input/SButton.h"
+#include "Misc/ConfigCacheIni.h"
 
 #define LOCTEXT_NAMESPACE "StatusBar"
 
 int32 UStatusBarSubsystem::MessageHandleCounter = 0;
 
+class SNewUserTipNotification : public SCompoundWidget
+{
+	SLATE_BEGIN_ARGS(SNewUserTipNotification)
+	{}
+	SLATE_END_ARGS()
+
+	SNewUserTipNotification()
+		: NewBadgeBrush(FStyleColors::Success)
+		, KeybindBackgroundBrush(FLinearColor::Transparent, 6.0f, FStyleColors::ForegroundHover, 1.5f)
+	{}
+
+	static void Show(TSharedPtr<SWindow> InParentWindow)
+	{
+		if(!ActiveNotification.IsValid())
+		{
+			TSharedRef<SNewUserTipNotification> ActiveNotificationRef = 
+				SNew(SNewUserTipNotification);
+
+			ActiveNotification = ActiveNotificationRef;
+			ParentWindow = InParentWindow;
+			InParentWindow->AddOverlaySlot()
+				.VAlign(VAlign_Bottom)
+				.HAlign(HAlign_Left)
+				.Padding(FMargin(20.0f, 20.0f, 10.0f, 50.f))
+				[
+					ActiveNotificationRef
+				];
+		}
+
+	}
+
+	static void Dismiss()
+	{
+		TSharedPtr<SNewUserTipNotification> ActiveNotificationPin = ActiveNotification.Pin();
+		if (ParentWindow.IsValid() && ActiveNotificationPin.IsValid())
+		{
+			ParentWindow.Pin()->RemoveOverlaySlot(ActiveNotificationPin.ToSharedRef());
+		}
+
+		ParentWindow.Reset();
+
+		ActiveNotification.Reset();
+	}
+
+	void Construct(const FArguments& InArgs)
+	{
+		ChildSlot
+		[
+			SNew(SBox)
+			.WidthOverride(350.0f)
+			.HeightOverride(128.0f)
+			[
+				SNew(SNotificationBackground)
+				.Padding(FMargin(16, 8))
+				[
+					SNew(SHorizontalBox)
+					+ SHorizontalBox::Slot()
+					.Padding(0.0f, 6.0f, 0.0f, 0.0f)
+					.VAlign(VAlign_Top)
+					.AutoWidth()
+					[
+						SNew(SBorder)
+						.Padding(FMargin(11,4))
+						.BorderImage(&NewBadgeBrush)
+						.ForegroundColor(FStyleColors::ForegroundInverted)
+						[
+							SNew(STextBlock)
+							.Text(LOCTEXT("NewBadge", "New"))
+							.TextStyle(FAppStyle::Get(), "SmallButtonText")
+							.ColorAndOpacity(FSlateColor::UseForeground())
+						]
+					]
+					+ SHorizontalBox::Slot()
+					.AutoWidth()
+					.Padding(16.0f, 8.0f, 0.0f, 0.0f)
+					[
+						SNew(SVerticalBox)
+						+ SVerticalBox::Slot()
+						.AutoHeight()
+						.Padding(0.0f, 0.0f, 0.0f, 0.0f)
+						[
+							SNew(STextBlock)
+							.Font(FAppStyle::Get().GetFontStyle("NotificationList.FontBold"))
+							.Text(LOCTEXT("ContentDrawerTipTitle", "Content Drawer"))
+							.ColorAndOpacity(FStyleColors::ForegroundHover)
+						]
+						+ SVerticalBox::Slot()
+						.AutoHeight()
+						.Padding(0.0f, 12.0f, 0.0f, 0.0f)
+						[
+							SNew(SHorizontalBox)
+							+ SHorizontalBox::Slot()
+							.AutoWidth()
+							[
+								SNew(SBorder)
+								.Padding(FMargin(20, 4))
+								.BorderImage(&KeybindBackgroundBrush)
+								[
+									SNew(STextBlock)
+									.TextStyle(FAppStyle::Get(), "DialogButtonText")
+									.Text(FText::FromString("CTRL"))
+									.ColorAndOpacity(FStyleColors::ForegroundHover)
+								]
+							]
+							+ SHorizontalBox::Slot()
+							.Padding(8.0f, 0.0f)
+							.AutoWidth()
+							.VAlign(VAlign_Center)
+							[
+								SNew(SImage)
+								.Image(FAppStyle::Get().GetBrush("Icons.Plus"))
+								.ColorAndOpacity(FStyleColors::ForegroundHover)
+							]
+							+ SHorizontalBox::Slot()
+							.AutoWidth()
+							[
+								SNew(SBorder)
+								.Padding(FMargin(20, 4))
+								.BorderImage(&KeybindBackgroundBrush)
+								[
+									SNew(STextBlock)
+									.TextStyle(FAppStyle::Get(), "DialogButtonText")
+									.Text(FText::FromString("SPACE"))
+									.ColorAndOpacity(FStyleColors::ForegroundHover)
+								]
+							]
+						]
+						+ SVerticalBox::Slot()
+						.AutoHeight()
+						.Padding(0.0f, 12.0f, 0.0f, 0.0f)
+						[
+							SNew(STextBlock)
+							.Text(LOCTEXT("ContentDrawerTipDesc", "Summon the content browser in\ncollapsable drawer."))
+							.ColorAndOpacity(FStyleColors::Foreground)
+						]
+					]
+					+ SHorizontalBox::Slot()
+					.Padding(0.0f, 0.0f, 0.0f, 0.0f)
+					.HAlign(HAlign_Right)
+					.VAlign(VAlign_Top)
+					[
+						SNew(SButton)
+						.ButtonStyle(FAppStyle::Get(), "SimpleButton")
+						.OnClicked_Lambda([]() { SNewUserTipNotification::Dismiss(); return FReply::Handled(); })
+						[
+							SNew(SImage)
+							.Image(FAppStyle::Get().GetBrush("Icons.X"))
+							.ColorAndOpacity(FSlateColor::UseForeground())
+						]
+					]
+				]
+			]
+		];
+	}
+
+private:
+	FSlateRoundedBoxBrush NewBadgeBrush;
+	FSlateRoundedBoxBrush KeybindBackgroundBrush;
+	static TWeakPtr<SNewUserTipNotification> ActiveNotification;
+	static TWeakPtr<SWindow> ParentWindow;
+};
+
+TWeakPtr<SNewUserTipNotification> SNewUserTipNotification::ActiveNotification;
+TWeakPtr<SWindow> SNewUserTipNotification::ParentWindow;
+
+
 void UStatusBarSubsystem::Initialize(FSubsystemCollectionBase& Collection)
 {
 	FSourceControlCommands::Register();
 
+	IMainFrameModule& MainFrameModule = IMainFrameModule::Get();
+	if (MainFrameModule.IsWindowInitialized())
+	{
+		CreateAndShowNewUserTipIfNeeded(MainFrameModule.GetParentWindow(), false);
+	}
+	else
+	{
+		MainFrameModule.OnMainFrameCreationFinished().AddUObject(this, &UStatusBarSubsystem::CreateAndShowNewUserTipIfNeeded);
+	}
+
+
 	FSlateNotificationManager::Get().SetProgressNotificationHandler(this);
 }
 
@@ -78,10 +262,30 @@ bool UStatusBarSubsystem::OpenContentBrowserDrawer()
 	return false;
 }
 
+bool UStatusBarSubsystem::ForceDismissDrawer()
+{
+	bool bWasDismissed = false;
+	for (auto StatusBar : StatusBars)
+	{
+		if (TSharedPtr<SStatusBar> StatusBarPinned = StatusBar.Value.Pin())
+		{
+			if (StatusBarPinned->IsDrawerOpened(StatusBarDrawerIds::ContentBrowser))
+			{
+				StatusBarPinned->DismissDrawer(nullptr);
+				bWasDismissed = true;
+			}
+		}
+	}
+
+	return bWasDismissed;
+}
+
 bool UStatusBarSubsystem::ToggleContentBrowser(TSharedRef<SWindow> ParentWindow)
 {
 	bool bWasDismissed = false;
 
+	SNewUserTipNotification::Dismiss();
+
 	for (auto StatusBar : StatusBars)
 	{
 		if (TSharedPtr<SStatusBar> StatusBarPinned = StatusBar.Value.Pin())
@@ -119,7 +323,7 @@ TSharedRef<SWidget> UStatusBarSubsystem::MakeStatusBarWidget(FName StatusBarName
 	ContentBrowserDrawer.GetDrawerContentDelegate.BindUObject(this, &UStatusBarSubsystem::OnGetContentBrowser);
 	ContentBrowserDrawer.OnDrawerOpenedDelegate.BindUObject(this, &UStatusBarSubsystem::OnContentBrowserOpened);
 	ContentBrowserDrawer.OnDrawerDismissedDelegate.BindUObject(this, &UStatusBarSubsystem::OnContentBrowserDismissed);
-	ContentBrowserDrawer.ButtonText = LOCTEXT("StatusBar_ContentBrowserButton", "Content Browser");
+	ContentBrowserDrawer.ButtonText = LOCTEXT("StatusBar_ContentBrowserButton", "Content Drawer");
 	ContentBrowserDrawer.ToolTipText = FText::Format(LOCTEXT("StatusBar_ContentBrowserDrawerToolTip", "Opens a temporary content browser above this status which will dismiss when it loses focus ({0})"), FGlobalEditorCommonCommands::Get().OpenContentBrowserDrawer->GetInputText());
 	ContentBrowserDrawer.Icon = FAppStyle::Get().GetBrush("ContentBrowser.TabIcon");
 
@@ -246,10 +450,67 @@ void UStatusBarSubsystem::CreateContentBrowserIfNeeded()
 		FContentBrowserConfig Config;
 		Config.bCanSetAsPrimaryBrowser = false;
 
-		StatusBarContentBrowser = ContentBrowserSingleton.CreateContentBrowserDrawer(Config);
+		TFunction<TSharedPtr<SDockTab>()> GetTab(
+			[this]() -> TSharedPtr<SDockTab>
+			{
+				for (auto StatusBar : StatusBars)
+				{
+					if (TSharedPtr<SStatusBar> StatusBarPinned = StatusBar.Value.Pin())
+					{
+						if (StatusBarPinned->IsDrawerOpened(StatusBarDrawerIds::ContentBrowser))
+						{
+							return StatusBarPinned->GetParentTab();
+						}
+					}
+				}
+
+				checkf(false, TEXT("If we get here somehow a content browser drawer is opened but no status bar claims it"));
+				return TSharedPtr<SDockTab>();
+			}
+		);
+		StatusBarContentBrowser = ContentBrowserSingleton.CreateContentBrowserDrawer(Config, GetTab);
 	}
 }
 
+void UStatusBarSubsystem::CreateAndShowNewUserTipIfNeeded(TSharedPtr<SWindow> ParentWindow, bool bIsNewProjectDialog)
+{
+	if(!bIsNewProjectDialog)
+	{
+		const FString StoreId = TEXT("Epic Games");
+		const FString SectionName = TEXT("Unreal Engine/Editor");
+		const FString KeyName = TEXT("LaunchTipShown");
+
+		const FString FallbackIniLocation = TEXT("/Script/UnrealEd.EditorSettings");
+		const FString FallbackIniKey = TEXT("LaunchTipShownFallback");
+
+		// Its important that this new user message does not appear after the first launch so we store it in a more permanent place
+		FString CurrentState = TEXT("0");
+		if (!FPlatformMisc::GetStoredValue(StoreId, SectionName, KeyName, CurrentState))
+		{
+			// As a fallback where the registry was not readable or writable, save a flag in the editor ini. This will be less permanent as the registry but will prevent 
+			// the notification from appearing on every launch
+			GConfig->GetString(*FallbackIniLocation, *FallbackIniKey, CurrentState, GEditorSettingsIni);
+
+
+		}
+
+		if(CurrentState != TEXT("1"))
+		{
+			SNewUserTipNotification::Show(ParentWindow);
+
+			// Write that we've shown the notification
+			if (!FPlatformMisc::SetStoredValue(StoreId, SectionName, KeyName, TEXT("1")))
+			{
+				// Use fallback
+				GConfig->SetString(*FallbackIniLocation, *FallbackIniKey, TEXT("1"), GEditorSettingsIni);
+			}
+		}
+	}
+
+	// Ignore the if the main frame gets recreated this session
+	IMainFrameModule::Get().OnMainFrameCreationFinished().RemoveAll(this);
+}
+
 TSharedPtr<SStatusBar> UStatusBarSubsystem::GetStatusBar(FName StatusBarName) const
 {
 	return StatusBars.FindRef(StatusBarName).Pin();
@@ -264,6 +525,8 @@ TSharedRef<SWidget> UStatusBarSubsystem::OnGetContentBrowser()
 
 void UStatusBarSubsystem::OnContentBrowserOpened(TSharedRef<SStatusBar>& StatusBarWithContentBrowser)
 {
+	SNewUserTipNotification::Dismiss();
+
 	// Dismiss any other content browser that is opened when one status bar opens it.  The content browser is a shared resource and shouldn't be in the layout twice
 	for (auto StatusBar : StatusBars)
 	{
diff --git a/Engine/Source/Editor/StatusBar/Public/StatusBarSubsystem.h b/Engine/Source/Editor/StatusBar/Public/StatusBarSubsystem.h
--- a/Engine/Source/Editor/StatusBar/Public/StatusBarSubsystem.h
+++ b/Engine/Source/Editor/StatusBar/Public/StatusBarSubsystem.h
@@ -79,6 +79,11 @@ public:
 	 */
 	bool OpenContentBrowserDrawer();
 
+	/**
+	 * Forces the drawer to dismiss. Usually it dismisses with focus. Only call this if there is some reason an open drawer would be invalid for the current state of the editor.
+	 */
+	bool ForceDismissDrawer();
+
 	/** 
 	 * Creates a new instance of a status bar widget
 	 *
@@ -121,16 +126,16 @@ private:
 	bool ToggleContentBrowser(TSharedRef<SWindow> ParentWindow);
 	void OnDebugConsoleClosed();
 	void CreateContentBrowserIfNeeded();
+	void CreateAndShowNewUserTipIfNeeded(TSharedPtr<SWindow> ParentWindow, bool bIsNewProjectDialog);
+
 	TSharedPtr<SStatusBar> GetStatusBar(FName StatusBarName) const;
 	TSharedRef<SWidget> OnGetContentBrowser();
 	void OnContentBrowserOpened(TSharedRef<SStatusBar>& StatusBarWithContentBrowser);
 	void OnContentBrowserDismissed(const TSharedPtr<SWidget>& NewlyFocusedWidget);
 	void HandleDeferredOpenContentBrowser(TSharedPtr<SWindow> ParentWindow);
-
 private:
 	TMap<FName, TWeakPtr<SStatusBar>> StatusBars;
 	TWeakPtr<SWidget> PreviousKeyboardFocusedWidget;
-
 	/** The floating content browser that is opened via the content browser button in the status bar */
 	TSharedPtr<SWidget> StatusBarContentBrowser;
 	static int32 MessageHandleCounter;
diff --git a/Engine/Source/Editor/StatusBar/StatusBar.Build.cs b/Engine/Source/Editor/StatusBar/StatusBar.Build.cs
--- a/Engine/Source/Editor/StatusBar/StatusBar.Build.cs
+++ b/Engine/Source/Editor/StatusBar/StatusBar.Build.cs
@@ -24,6 +24,7 @@ public class StatusBar : ModuleRules
 				"OutputLog",
 				"SourceControlWindows",
 				"EditorSubsystem",
+				"MainFrame",
             });
     }
 }
diff --git a/Engine/Source/Runtime/Slate/Private/Framework/Docking/TabManager.cpp b/Engine/Source/Runtime/Slate/Private/Framework/Docking/TabManager.cpp
--- a/Engine/Source/Runtime/Slate/Private/Framework/Docking/TabManager.cpp
+++ b/Engine/Source/Runtime/Slate/Private/Framework/Docking/TabManager.cpp
@@ -2498,6 +2498,17 @@ TSharedPtr<SDockTab> FGlobalTabmanager::GetMajorTabForTabManager(const TSharedRe
 	return TSharedPtr<SDockTab>();
 }
 
+TSharedPtr<FTabManager> FGlobalTabmanager::GetTabManagerForMajorTab(const TSharedPtr<SDockTab> DockTab) const
+{
+	const int32 Index = SubTabManagers.IndexOfByPredicate(FindByTab(DockTab.ToSharedRef()));
+	if (Index != INDEX_NONE)
+	{
+		return SubTabManagers[Index].TabManager.Pin();
+	}
+
+	return nullptr;
+}
+
 void FGlobalTabmanager::DrawAttentionToTabManager( const TSharedRef<FTabManager>& ChildManager )
 {
 	TSharedPtr<SDockTab> Tab = GetMajorTabForTabManager(ChildManager);
diff --git a/Engine/Source/Runtime/Slate/Private/Widgets/Notifications/SNotificationBackground.cpp b/Engine/Source/Runtime/Slate/Private/Widgets/Notifications/SNotificationBackground.cpp
new file mode 100644
--- /dev/null
+++ b/Engine/Source/Runtime/Slate/Private/Widgets/Notifications/SNotificationBackground.cpp
@@ -0,0 +1,61 @@
+#include "Widgets/Notifications/SNotificationBackground.h"
+#include "Styling/StyleColors.h"
+
+void SNotificationBackground::Construct(const FArguments& InArgs)
+{
+	WatermarkBrush = FAppStyle::Get().GetBrush("NotificationList.Watermark");
+	BorderBrush = FAppStyle::Get().GetBrush("NotificationList.ItemBackground_Border");
+
+	SetClipping(EWidgetClipping::ClipToBounds);
+
+	WatermarkTint = FStyleColors::Notifications.GetSpecifiedColor() * FLinearColor(1.15f, 1.15f, 1.15f, 1.0f);
+
+	SBorder::Construct(
+		SBorder::FArguments()
+		.BorderImage(FStyleDefaults::GetNoBrush())
+		.Padding(InArgs._Padding)
+		.BorderBackgroundColor(InArgs._BorderBackgroundColor)
+		.ColorAndOpacity(InArgs._ColorAndOpacity)
+		.DesiredSizeScale(InArgs._DesiredSizeScale)
+		[
+			InArgs._Content.Widget
+		]
+	);
+}
+
+int32 SNotificationBackground::OnPaint(const FPaintArgs& Args, const FGeometry& AllottedGeometry, const FSlateRect& MyCullingRect, FSlateWindowElementList& OutDrawElements, int32 LayerId, const FWidgetStyle& InWidgetStyle, bool bParentEnabled) const
+{
+	if (WatermarkBrush && WatermarkBrush->DrawAs != ESlateBrushDrawType::NoDrawType)
+	{
+		// The watermark should be 250% bigger than the size of the notification to get the desired effect
+		const float SizeY = AllottedGeometry.GetLocalSize().Y;
+		const float WatermarkSize = FMath::Clamp(SizeY * 2.50f, 130.0f, 220.0f);
+
+
+		FSlateDrawElement::MakeBox(
+			OutDrawElements,
+			LayerId,
+			AllottedGeometry.ToPaintGeometry(),
+			BorderBrush,
+			ESlateDrawEffect::None,
+			BorderBrush->GetTint(InWidgetStyle) * InWidgetStyle.GetColorAndOpacityTint());
+
+		OutDrawElements.PushClip(FSlateClippingZone(MyCullingRect.InsetBy(FMargin(2))));
+
+		FSlateDrawElement::MakeBox(
+			OutDrawElements,
+			LayerId + 1,
+			AllottedGeometry.ToPaintGeometry(FVector2D(4.0f, 4.0f), FVector2D(WatermarkSize, WatermarkSize)),
+			WatermarkBrush,
+			ESlateDrawEffect::None,
+			WatermarkTint * InWidgetStyle.GetColorAndOpacityTint() * GetBorderBackgroundColor().GetColor(InWidgetStyle));
+
+		OutDrawElements.PopClip();
+	}
+
+
+	int32 OutLayerId = SBorder::OnPaint(Args, AllottedGeometry, MyCullingRect, OutDrawElements, LayerId+2, InWidgetStyle, bParentEnabled);
+
+	return OutLayerId;
+}
+
diff --git a/Engine/Source/Runtime/Slate/Private/Widgets/Notifications/SNotificationList.cpp b/Engine/Source/Runtime/Slate/Private/Widgets/Notifications/SNotificationList.cpp
--- a/Engine/Source/Runtime/Slate/Private/Widgets/Notifications/SNotificationList.cpp
+++ b/Engine/Source/Runtime/Slate/Private/Widgets/Notifications/SNotificationList.cpp
@@ -17,6 +17,7 @@
 #include "Widgets/Notifications/INotificationWidget.h"
 #include "Framework/Notifications/NotificationManager.h"
 #include "Styling/StyleColors.h"
+#include "Widgets/Notifications/SNotificationBackground.h"
 
 /////////////////////////////////////////////////
 // SNotificationExtendable
@@ -317,77 +318,6 @@ protected:
 };
 
 
-class SNotificationBackground : public SBorder
-{
-public:
-	SLATE_BEGIN_ARGS(SNotificationBackground) {}
-		SLATE_ATTRIBUTE(FMargin, Padding)
-		/** ColorAndOpacity is the color and opacity of content in the border */
-		SLATE_ATTRIBUTE(FLinearColor, ColorAndOpacity)
-		/** BorderBackgroundColor refers to the actual color and opacity of the supplied border image.*/
-		SLATE_ATTRIBUTE(FSlateColor, BorderBackgroundColor)
-
-		SLATE_ATTRIBUTE(FVector2D, DesiredSizeScale)
-
-		SLATE_DEFAULT_SLOT(FArguments, Content)
-	SLATE_END_ARGS()
-
-	void Construct(const FArguments& InArgs)
-	{
-		WatermarkBrush = FAppStyle::Get().GetBrush("NotificationList.Watermark");
-		BorderBrush = FAppStyle::Get().GetBrush("NotificationList.ItemBackground_Border");
-
-		WatermarkTint = FStyleColors::Notifications.GetSpecifiedColor() * FLinearColor(1.15f, 1.15f, 1.15f, 1.0f);
-
-		SBorder::Construct(
-			SBorder::FArguments()
-			.BorderImage(FAppStyle::Get().GetBrush("NotificationList.ItemBackground"))
-			.Padding(InArgs._Padding)
-			.BorderBackgroundColor(InArgs._BorderBackgroundColor)
-			.ColorAndOpacity(InArgs._ColorAndOpacity)
-			.DesiredSizeScale(InArgs._DesiredSizeScale)
-			[
-				InArgs._Content.Widget
-			]
-		);
-	}
-
-	virtual int32 OnPaint(const FPaintArgs& Args, const FGeometry& AllottedGeometry, const FSlateRect& MyCullingRect, FSlateWindowElementList& OutDrawElements, int32 LayerId, const FWidgetStyle& InWidgetStyle, bool bParentEnabled) const override
-	{
-		if (WatermarkBrush && WatermarkBrush->DrawAs != ESlateBrushDrawType::NoDrawType)
-		{
-			// The watermark should be 250% bigger than the size of the notification to get the desired effect
-			const float SizeY = AllottedGeometry.GetLocalSize().Y;
-			const float WatermarkSize = FMath::Clamp(SizeY * 2.50f, 130.0f, 220.0f);
-	
-			FSlateDrawElement::MakeBox(
-				OutDrawElements,
-				LayerId + 1,
-				AllottedGeometry.ToPaintGeometry(FVector2D(4.0f, 4.0f), FVector2D(WatermarkSize, WatermarkSize)),
-				WatermarkBrush,
-				ESlateDrawEffect::None,
-				WatermarkTint * InWidgetStyle.GetColorAndOpacityTint() * GetBorderBackgroundColor().GetColor(InWidgetStyle));
-		}
-
-
-		int32 OutLayerId = SBorder::OnPaint(Args, AllottedGeometry, MyCullingRect, OutDrawElements, LayerId, InWidgetStyle, bParentEnabled);
-
-		FSlateDrawElement::MakeBox(
-			OutDrawElements,
-			OutLayerId,
-			AllottedGeometry.ToPaintGeometry(),
-			BorderBrush,
-			ESlateDrawEffect::None,
-			BorderBrush->GetTint(InWidgetStyle) * InWidgetStyle.GetColorAndOpacityTint());
-
-		return OutLayerId+1;
-	}
-	
-private:
-	const FSlateBrush* WatermarkBrush = nullptr;
-	const FSlateBrush* BorderBrush = nullptr;
-	FLinearColor WatermarkTint;
-};
 
 /////////////////////////////////////////////////
 // SNotificationItemImpl
@@ -454,26 +384,14 @@ public:
 
 		ChildSlot
 		[
-/*
-			SNew(SOverlay)
-			+ SOverlay::Slot()
-			[*/
-				SNew(SNotificationBackground)
-				.Padding(FMargin(16, 8))
-				.BorderBackgroundColor(this, &SNotificationItemImpl::GetContentColor)
-				.ColorAndOpacity(this, &SNotificationItemImpl::GetContentColorRaw)
-				.DesiredSizeScale(this, &SNotificationItemImpl::GetItemScale)
-				[
-					ConstructInternals(InArgs)
-				]
-/*
-			]
-			+SOverlay::Slot()
+			SNew(SNotificationBackground)
+			.Padding(FMargin(16, 8))
+			.BorderBackgroundColor(this, &SNotificationItemImpl::GetContentColor)
+			.ColorAndOpacity(this, &SNotificationItemImpl::GetContentColorRaw)
+			.DesiredSizeScale(this, &SNotificationItemImpl::GetItemScale)
 			[
-				SNew(SImage)
-				.Image(FAppStyle::Get().GetBrush("NotificationList.ItemBackground_Border"))
-				.Visibility(EVisibility::HitTestInvisible)
-			]*/
+				ConstructInternals(InArgs)
+			]
 		];
 	}
 	
@@ -721,9 +639,6 @@ public:
 		if (NotificationWidget->UseNotificationBackground())
 		{
 			Internals =
-				/*SNew(SOverlay)
-				+ SOverlay::Slot()
-				[*/
 					SNew(SNotificationBackground)
 					.Padding(FMargin(16, 8))
 					.BorderBackgroundColor(this, &SNotificationItemExternalImpl::GetContentColor)
@@ -736,13 +651,6 @@ public:
 							NotificationWidget->AsWidget()
 						]
 					];
-			/*	]
-				+SOverlay::Slot()
-				[
-					SNew(SImage)
-					.Image(FAppStyle::Get().GetBrush("NotificationList.ItemBackground_Border"))
-					.Visibility(EVisibility::HitTestInvisible)
-				];*/
 		}
 		else
 		{
diff --git a/Engine/Source/Runtime/Slate/Public/Framework/Docking/TabManager.h b/Engine/Source/Runtime/Slate/Public/Framework/Docking/TabManager.h
--- a/Engine/Source/Runtime/Slate/Public/Framework/Docking/TabManager.h
+++ b/Engine/Source/Runtime/Slate/Public/Framework/Docking/TabManager.h
@@ -1164,6 +1164,12 @@ public:
 	/** Gets the major tab for the manager */
 	TSharedPtr<SDockTab> GetMajorTabForTabManager(const TSharedRef<FTabManager>& ChildManager);
 
+	/** 
+	 * Gets the tab manager that a major tab owns. The returned tab manager is the tab manager that manages the minor tabs for a major tab.
+	 * Note: this is not the same as DockTab->GetTabManager(). That function returns the tab manager the major tab is in
+	 */
+	TSharedPtr<FTabManager> GetTabManagerForMajorTab(const TSharedPtr<SDockTab> DockTab) const;
+
 	/** Draw the user's attention to a child tab manager */
 	void DrawAttentionToTabManager( const TSharedRef<FTabManager>& ChildManager );
 
diff --git a/Engine/Source/Runtime/Slate/Public/Widgets/Notifications/SNotificationBackground.h b/Engine/Source/Runtime/Slate/Public/Widgets/Notifications/SNotificationBackground.h
new file mode 100644
--- /dev/null
+++ b/Engine/Source/Runtime/Slate/Public/Widgets/Notifications/SNotificationBackground.h
@@ -0,0 +1,31 @@
+// Copyright Epic Games, Inc. All Rights Reserved.
+
+#pragma once
+
+#include "Widgets/Layout/SBorder.h"
+
+struct FSlateBrush;
+
+class SNotificationBackground : public SBorder
+{
+public:
+	SLATE_BEGIN_ARGS(SNotificationBackground) 
+		: _ColorAndOpacity(FLinearColor::White)
+		, _BorderBackgroundColor(FLinearColor::White)
+		, _DesiredSizeScale(FVector2D(1.0f, 1.0f))
+	{}
+		SLATE_ATTRIBUTE(FMargin, Padding)
+		SLATE_ATTRIBUTE(FLinearColor, ColorAndOpacity)
+		SLATE_ATTRIBUTE(FSlateColor, BorderBackgroundColor)
+		SLATE_ATTRIBUTE(FVector2D, DesiredSizeScale)
+		SLATE_DEFAULT_SLOT(FArguments, Content)
+	SLATE_END_ARGS()
+
+	SLATE_API void Construct(const FArguments& InArgs);
+	SLATE_API virtual int32 OnPaint(const FPaintArgs& Args, const FGeometry& AllottedGeometry, const FSlateRect& MyCullingRect, FSlateWindowElementList& OutDrawElements, int32 LayerId, const FWidgetStyle& InWidgetStyle, bool bParentEnabled) const override;
+
+private:
+	const FSlateBrush* WatermarkBrush = nullptr;
+	const FSlateBrush* BorderBrush = nullptr;
+	FLinearColor WatermarkTint;
+};
\ No newline at end of file
