The output log can now be summoned from a drawer in each asset editor status bar. The default sequence mirrors -game behavior where '~' will toggle between focusing the console input box, opening the drawer, and restoring focus to the previously focused UI element. This behavior can be disabled in settings to prevent the output log from opening as a drawer with '~'. Additionally a separate shortcut can be bound in order to activate the output log drawer independently of console input box focus.

diff --git a/Engine/Source/Developer/OutputLog/Private/OutputLogModule.cpp b/Engine/Source/Developer/OutputLog/Private/OutputLogModule.cpp
--- a/Engine/Source/Developer/OutputLog/Private/OutputLogModule.cpp
+++ b/Engine/Source/Developer/OutputLog/Private/OutputLogModule.cpp
@@ -67,23 +67,26 @@ private:
 	TArray< TSharedPtr<FOutputLogMessage> > Messages;
 };
 
-/** Our global output log app spawner */
-static TSharedPtr<FOutputLogHistory> OutputLogHistory;
 
-/** Our global active output log */
-static TWeakPtr<SOutputLog> OutputLog;
-
-TSharedRef<SDockTab> SpawnOutputLog(const FSpawnTabArgs& Args)
+TSharedRef<SDockTab> FOutputLogModule::SpawnOutputLogTab(const FSpawnTabArgs& Args)
 {
-	return SNew(SDockTab)
+	TSharedRef<SOutputLog> NewLog = SNew(SOutputLog, false).Messages(OutputLogHistory->GetMessages());
+
+	OutputLog = NewLog;
+
+	TSharedRef<SDockTab> NewTab = SNew(SDockTab)
 		.TabRole(ETabRole::NomadTab)
 		.Label(NSLOCTEXT("OutputLog", "TabTitle", "Output Log"))
 		[
-			SAssignNew(OutputLog, SOutputLog).Messages(OutputLogHistory->GetMessages())
+			NewLog
 		];
+
+	OutputLogTab = NewTab;
+
+	return NewTab;
 }
 
-TSharedRef<SDockTab> SpawnDeviceOutputLog(const FSpawnTabArgs& Args)
+TSharedRef<SDockTab> FOutputLogModule::SpawnDeviceOutputLogTab(const FSpawnTabArgs& Args)
 {
 	return SNew(SDockTab)
 		.TabRole(ETabRole::NomadTab)
@@ -95,13 +98,13 @@ TSharedRef<SDockTab> SpawnDeviceOutputLog(const FSpawnTabArgs& Args)
 
 void FOutputLogModule::StartupModule()
 {
-	FGlobalTabmanager::Get()->RegisterNomadTabSpawner(OutputLogModule::OutputLogTabName, FOnSpawnTab::CreateStatic(&SpawnOutputLog))
+	FGlobalTabmanager::Get()->RegisterNomadTabSpawner(OutputLogModule::OutputLogTabName, FOnSpawnTab::CreateRaw(this, &FOutputLogModule::SpawnOutputLogTab))
 		.SetDisplayName(NSLOCTEXT("UnrealEditor", "OutputLogTab", "Output Log"))
 		.SetTooltipText(NSLOCTEXT("UnrealEditor", "OutputLogTooltipText", "Open the Output Log tab."))
 		.SetGroup(WorkspaceMenu::GetMenuStructure().GetDeveloperToolsLogCategory())
 		.SetIcon(FSlateIcon(FEditorStyle::GetStyleSetName(), "Log.TabIcon"));
 
-	FGlobalTabmanager::Get()->RegisterNomadTabSpawner(OutputLogModule::DeviceOutputLogTabName, FOnSpawnTab::CreateStatic(&SpawnDeviceOutputLog))
+	FGlobalTabmanager::Get()->RegisterNomadTabSpawner(OutputLogModule::DeviceOutputLogTabName, FOnSpawnTab::CreateRaw(this, &FOutputLogModule::SpawnDeviceOutputLogTab))
 		.SetDisplayName(NSLOCTEXT("UnrealEditor", "DeviceOutputLogTab", "Device Output Log"))
 		.SetTooltipText(NSLOCTEXT("UnrealEditor", "DeviceOutputLogTooltipText", "Open the Device Output Log tab."))
 		.SetGroup(WorkspaceMenu::GetMenuStructure().GetDeveloperToolsLogCategory())
@@ -129,16 +132,40 @@ void FOutputLogModule::ShutdownModule()
 	OutputLogHistory.Reset();
 }
 
-TSharedRef< SWidget > FOutputLogModule::MakeConsoleInputBox(TSharedPtr<SMultiLineEditableTextBox>& OutExposedEditableTextBox, const FSimpleDelegate& OnCloseConsole) const
+FOutputLogModule& FOutputLogModule::Get()
 {
-	TSharedRef< SConsoleInputBox > NewConsoleInputBox =
+	static const FName OutputLog("OutputLog");
+
+	return FModuleManager::Get().LoadModuleChecked<FOutputLogModule>(OutputLog);
+}
+
+TSharedRef<SWidget> FOutputLogModule::MakeConsoleInputBox(TSharedPtr<SMultiLineEditableTextBox>& OutExposedEditableTextBox, const FSimpleDelegate& OnCloseConsole, const FSimpleDelegate& OnConsoleCommandExecuted) const
+{
+	TSharedRef<SConsoleInputBox> NewConsoleInputBox =
 		SNew(SConsoleInputBox)
-		.OnCloseConsole(OnCloseConsole);
+		.OnCloseConsole(OnCloseConsole)
+		.OnConsoleCommandExecuted(OnConsoleCommandExecuted);
 
 	OutExposedEditableTextBox = NewConsoleInputBox->GetEditableTextBox();
 	return NewConsoleInputBox;
 }
 
+TSharedRef<SWidget> FOutputLogModule::MakeOutputLogDrawerWidget(const FSimpleDelegate& OnCloseConsole)
+{
+	TSharedPtr<SOutputLog> OutputLogDrawerPinned = OutputLogDrawer.Pin();
+
+	if (!OutputLogDrawerPinned.IsValid())
+	{
+		OutputLogDrawerPinned = 
+			SNew(SOutputLog, true)
+			.OnCloseConsole(OnCloseConsole)
+			.Messages(OutputLogHistory->GetMessages());
+		OutputLogDrawer = OutputLogDrawerPinned;
+	}
+
+	return OutputLogDrawerPinned.ToSharedRef();
+}
+
 void FOutputLogModule::ToggleDebugConsoleForWindow(const TSharedRef<SWindow>& Window, const EDebugConsoleStyle::Type InStyle, const FDebugConsoleDelegates& DebugConsoleDelegates)
 {
 	bool bShouldOpen = true;
@@ -169,7 +196,7 @@ void FOutputLogModule::ToggleDebugConsoleForWindow(const TSharedRef<SWindow>& Wi
 	}
 
 	TSharedPtr<SDockTab> ActiveTab = FGlobalTabmanager::Get()->GetActiveTab();
-	if (ActiveTab.IsValid() && ActiveTab->GetLayoutIdentifier() == FTabId(OutputLogModule::OutputLogTabName))
+	if (ActiveTab.IsValid() && ActiveTab == OutputLogTab)
 	{
 		FGlobalTabmanager::Get()->DrawAttention(ActiveTab.ToSharedRef());
 		bShouldOpen = false;
@@ -214,19 +241,42 @@ void FOutputLogModule::CloseDebugConsole()
 
 void FOutputLogModule::ClearOnPIE(const bool bIsSimulating)
 {
-	if (OutputLog.IsValid())
-	{
-		bool bClearOnPIEEnabled = false;
-		GConfig->GetBool(TEXT("/Script/UnrealEd.EditorPerProjectUserSettings"), TEXT("bEnableOutputLogClearOnPIE"), bClearOnPIEEnabled, GEditorPerProjectIni);
+	bool bClearOnPIEEnabled = false;
+	GConfig->GetBool(TEXT("/Script/UnrealEd.EditorPerProjectUserSettings"), TEXT("bEnableOutputLogClearOnPIE"), bClearOnPIEEnabled, GEditorPerProjectIni);
 
-		if (bClearOnPIEEnabled)
+	if (bClearOnPIEEnabled)
+	{
+		if(TSharedPtr<SOutputLog> OutputLogPinned = OutputLog.Pin())
 		{
-			TSharedPtr<SOutputLog> OutputLogShared = OutputLog.Pin();
+			if (OutputLogPinned->CanClearLog())
+			{
+				OutputLogPinned->OnClearLog();
+			}
+		}
 
-			if (OutputLogShared->CanClearLog())
+		if (TSharedPtr<SOutputLog> OutputLogPinned = OutputLogDrawer.Pin())
+		{
+			if (OutputLogPinned->CanClearLog())
 			{
-				OutputLogShared->OnClearLog();
+				OutputLogPinned->OnClearLog();
 			}
 		}
 	}
 }
+
+void FOutputLogModule::FocusOutputLogConsoleBox(const TSharedRef<SWidget> OutputLogToFocus)
+{
+	if (OutputLog == OutputLogToFocus)
+	{
+		OutputLog.Pin()->FocusConsoleCommandBox();
+	}
+	else if (OutputLogDrawer == OutputLogToFocus)
+	{
+		OutputLogDrawer.Pin()->FocusConsoleCommandBox();
+	}
+}
+
+const TSharedPtr<SWidget> FOutputLogModule::GetOutputLog() const
+{
+	return OutputLog.Pin();
+}
diff --git a/Engine/Source/Developer/OutputLog/Private/SOutputLog.cpp b/Engine/Source/Developer/OutputLog/Private/SOutputLog.cpp
--- a/Engine/Source/Developer/OutputLog/Private/SOutputLog.cpp
+++ b/Engine/Source/Developer/OutputLog/Private/SOutputLog.cpp
@@ -26,6 +26,10 @@
 #include "Misc/CoreDelegates.h"
 #include "HAL/PlatformOutputDevices.h"
 #include "HAL/FileManager.h"
+#include "Widgets/Input/SButton.h"
+#include "Framework/Docking/TabManager.h"
+#include "Widgets/Docking/SDockTab.h"
+#include "OutputLogModule.h"
 
 #define LOCTEXT_NAMESPACE "SOutputLog"
 /** Expression context to test the given messages against the current text filter */
@@ -88,7 +92,7 @@ void SConsoleInputBox::Construct(const FArguments& InArgs)
 			[
 				SNew(SComboButton)
 				.IsEnabled(this, &SConsoleInputBox::IsCommandExecutorMenuEnabled)
-				.ComboButtonStyle(FAppStyle::Get(), "DebugConsole.ComboButton")
+				.ComboButtonStyle(FAppStyle::Get(), "SimpleComboButton")
 				.ContentPadding(0)
 				.OnGetMenuContent(this, &SConsoleInputBox::GetCommandExecutorMenuContent)
 				.ButtonContent()
@@ -835,8 +839,10 @@ FOutputLogTextLayoutMarshaller::FOutputLogTextLayoutMarshaller(TArray< TSharedPt
 }
 
 BEGIN_SLATE_FUNCTION_BUILD_OPTIMIZATION
-void SOutputLog::Construct( const FArguments& InArgs )
+void SOutputLog::Construct( const FArguments& InArgs, bool bIsDrawerOutputLog)
 {
+	bIsInDrawer = bIsDrawerOutputLog;
+
 	// Build list of available log categories from historical logs
 	for (const auto& Message : InArgs._Messages)
 	{
@@ -856,126 +862,109 @@ void SOutputLog::Construct( const FArguments& InArgs )
 		.ContextMenuExtender(this, &SOutputLog::ExtendTextBoxMenu);
 
 	ChildSlot
+	.Padding(3)
 	[
-		SNew(SBorder)
-		.Padding(3)
-		.BorderImage(FEditorStyle::GetBrush("ToolPanel.GroupBorder"))
-		[
-			SNew(SVerticalBox)
+		SNew(SVerticalBox)
 
-			// Output Log Filter
-			+SVerticalBox::Slot()
-			.AutoHeight()
-			.Padding(FMargin(0.0f, 0.0f, 0.0f, 4.0f))
+		// Output Log Filter
+		+SVerticalBox::Slot()
+		.AutoHeight()
+		.Padding(FMargin(0.0f, 4.0f, 0.0f, 4.0f))
+		[
+			SNew(SHorizontalBox)
+			+SHorizontalBox::Slot()
+			.Padding(0, 0, 4, 0)
+			.FillWidth(.65f)
 			[
-				SNew(SHorizontalBox)
-			
-				+SHorizontalBox::Slot()
-				.AutoWidth()
+				SAssignNew(FilterTextBox, SSearchBox)
+				.HintText(LOCTEXT("SearchLogHint", "Search Log"))
+				.OnTextChanged(this, &SOutputLog::OnFilterTextChanged)
+				.OnTextCommitted(this, &SOutputLog::OnFilterTextCommitted)
+				.DelayChangeNotificationsWhileTyping(true)
+			]
+			+SHorizontalBox::Slot()
+			.AutoWidth()
+			.HAlign(HAlign_Left)
+			[
+				SNew(SComboButton)
+				.ComboButtonStyle(FAppStyle::Get(), "SimpleComboButton")
+				.ToolTipText(LOCTEXT("AddFilterToolTip", "Add an output log filter."))
+				.OnGetMenuContent(this, &SOutputLog::MakeAddFilterMenu)
+				.ButtonContent()
 				[
-					SNew(SComboButton)
-					.ComboButtonStyle(FEditorStyle::Get(), "GenericFilters.ComboButtonStyle")
-					.ForegroundColor(FLinearColor::White)
-					.ContentPadding(0)
-					.ToolTipText(LOCTEXT("AddFilterToolTip", "Add an output log filter."))
-					.OnGetMenuContent(this, &SOutputLog::MakeAddFilterMenu)
-					.HasDownArrow(true)
-					.ContentPadding(FMargin(1, 0))
-					.ButtonContent()
+					SNew(SHorizontalBox)
+					+SHorizontalBox::Slot()
+					.AutoWidth()
 					[
-						SNew(SHorizontalBox)
-
-						+SHorizontalBox::Slot()
-						.AutoWidth()
-						[
-							SNew(STextBlock)
-							.TextStyle(FEditorStyle::Get(), "GenericFilters.TextStyle")
-							.Font(FEditorStyle::Get().GetFontStyle("FontAwesome.9"))
-							.Text(FText::FromString(FString(TEXT("\xf0b0"))) /*fa-filter*/)
-						]
-
-						+SHorizontalBox::Slot()
-						.AutoWidth()
-						.Padding(2, 0, 0, 0)
-						[
-							SNew(STextBlock)
-							.TextStyle(FEditorStyle::Get(), "GenericFilters.TextStyle")
-							.Text(LOCTEXT("Filters", "Filters"))
-						]
+						SNew(SImage)
+						.Image(FAppStyle::Get().GetBrush("Icons.Filter"))
+						.ColorAndOpacity(FSlateColor::UseForeground())
+					]
+					+SHorizontalBox::Slot()
+					.AutoWidth()
+					.Padding(2, 0, 0, 0)
+					[
+						SNew(STextBlock)
+						.Text(LOCTEXT("Filters", "Filters"))
+						.ColorAndOpacity(FSlateColor::UseForeground())
 					]
-				]
-
-				+SHorizontalBox::Slot()
-				.Padding(4, 1, 0, 0)
-				[
-					SAssignNew(FilterTextBox, SSearchBox)
-					.HintText(LOCTEXT("SearchLogHint", "Search Log"))
-					.OnTextChanged(this, &SOutputLog::OnFilterTextChanged)
-					.OnTextCommitted(this, &SOutputLog::OnFilterTextCommitted)
-					.DelayChangeNotificationsWhileTyping(true)
 				]
 			]
-
-			// Output log area
-			+SVerticalBox::Slot()
-			.FillHeight(1)
+			+SHorizontalBox::Slot()
+			.HAlign(HAlign_Right)
+			.VAlign(VAlign_Center)
+			.Padding(4, 0)
 			[
-				MessagesTextBox.ToSharedRef()
+				CreateDrawerDockButton()
 			]
-
-			// The console input box
-			+SVerticalBox::Slot()
-			.AutoHeight()
+			+ SHorizontalBox::Slot()
+			.HAlign(HAlign_Right)
+			.AutoWidth()
 			[
-				SNew(SHorizontalBox)
-
-				+SHorizontalBox::Slot()
-				.FillWidth(1.f)
-				.VAlign(VAlign_Center)
-				.Padding(FMargin(0.0f, 1.0f, 0.0f, 0.0f))
+				SNew(SComboButton)
+				.ComboButtonStyle(FAppStyle::Get(), "SimpleComboButton")
+				.OnGetMenuContent(this, &SOutputLog::GetViewButtonContent)
+				.ButtonContent()
 				[
-					SNew(SBox)
-					.MaxDesiredHeight(180.0f)
+					SNew(SHorizontalBox)
+					+ SHorizontalBox::Slot()
+					.AutoWidth()
+					.VAlign(VAlign_Center)
 					[
-						SNew(SConsoleInputBox)
-						.OnConsoleCommandExecuted(this, &SOutputLog::OnConsoleCommandExecuted)
-
-						// Always place suggestions above the input line for the output log widget
-						.SuggestionListPlacement(MenuPlacement_AboveAnchor)
+						SNew(SImage)
+						.Image(FAppStyle::Get().GetBrush("Icons.Settings"))
+						.ColorAndOpacity(FSlateColor::UseForeground())
 					]
-				]
-
-				+SHorizontalBox::Slot()
-				.AutoWidth()
-				.Padding(4, 0, 0, 0)
-				[
-					SAssignNew(ViewOptionsComboButton, SComboButton)
-					.ContentPadding(0)
-					.ForegroundColor( this, &SOutputLog::GetViewButtonForegroundColor )
-					.ButtonStyle( FEditorStyle::Get(), "ToggleButton" ) // Use the tool bar item style for this button
-					.OnGetMenuContent( this, &SOutputLog::GetViewButtonContent )
-					.ButtonContent()
+					+ SHorizontalBox::Slot()
+					.AutoWidth()
+					.Padding(2, 0, 0, 0)
+					.VAlign(VAlign_Center)
 					[
-						SNew(SHorizontalBox)
- 
-						+SHorizontalBox::Slot()
-						.AutoWidth()
-						.VAlign(VAlign_Center)
-						[
-							SNew(SImage).Image( FEditorStyle::GetBrush("GenericViewButton") )
-						]
- 
-						+SHorizontalBox::Slot()
-						.AutoWidth()
-						.Padding(2, 0, 0, 0)
-						.VAlign(VAlign_Center)
-						[
-							SNew(STextBlock).Text( LOCTEXT("ViewButton", "View Options") )
-						]
+						SNew(STextBlock)
+						.Text(LOCTEXT("SettingsButton", "Settings"))
+						.ColorAndOpacity(FSlateColor::UseForeground())
 					]
 				]
 			]
 		]
+
+		// Output log area
+		+SVerticalBox::Slot()
+		.FillHeight(1)
+		[
+			MessagesTextBox.ToSharedRef()
+		]
+
+		// The console input box
+		+SVerticalBox::Slot()
+		.AutoHeight()
+		[
+			SAssignNew(ConsoleInputBox, SConsoleInputBox)
+			.OnConsoleCommandExecuted(this, &SOutputLog::OnConsoleCommandExecuted)
+			.OnCloseConsole(InArgs._OnCloseConsole)
+			// Always place suggestions above the input line for the output log widget
+			.SuggestionListPlacement(MenuPlacement_AboveAnchor) 
+		]
 	];
 
 	GLog->AddOutputDevice(this);
@@ -1107,14 +1096,6 @@ void SOutputLog::Serialize(const TCHAR* V, ELogVerbosity::Type Verbosity, const
 	MessagesTextMarshaller->AppendPendingMessage(V, Verbosity, Category);
 }
 
-FSlateColor SOutputLog::GetViewButtonForegroundColor() const
-{
-	static const FName InvertedForegroundName("InvertedForeground");
-	static const FName DefaultForegroundName("DefaultForeground");
-
-	return ViewOptionsComboButton->IsHovered() ? FEditorStyle::GetSlateColor(InvertedForegroundName) : FEditorStyle::GetSlateColor(DefaultForegroundName);
-}
-
 void SOutputLog::ExtendTextBoxMenu(FMenuBuilder& Builder)
 {
 	FUIAction ClearOutputLogAction(
@@ -1150,6 +1131,11 @@ bool SOutputLog::CanClearLog() const
 	return MessagesTextMarshaller->GetNumMessages() > 0;
 }
 
+void SOutputLog::FocusConsoleCommandBox()
+{
+	FSlateApplication::Get().SetKeyboardFocus(ConsoleInputBox->GetEditableTextBox(), EFocusCause::SetDirectly);
+}
+
 void SOutputLog::OnConsoleCommandExecuted()
 {
 	// Submit pending messages when executing a command to keep the log feeling responsive to input
@@ -1466,6 +1452,46 @@ TSharedRef<SWidget> SOutputLog::GetViewButtonContent()
 	return MenuBuilder.MakeWidget();
 }
 
+TSharedRef<SWidget> SOutputLog::CreateDrawerDockButton()
+{
+	if (bIsInDrawer)
+	{
+		return
+			SNew(SButton)
+			.ButtonStyle(FAppStyle::Get(), "SimpleButton")
+			.ToolTipText(LOCTEXT("DockInLayout_Tooltip", "Docks this output log in the current layout.\nThe drawer will still be usable as a temporary log."))
+			.ContentPadding(FMargin(1, 0))
+			.Visibility_Lambda(
+				[]()
+				{
+					return FOutputLogModule::Get().GetOutputLogTab() == nullptr ? EVisibility::Visible : EVisibility::Hidden;
+				})
+			.OnClicked(this, &SOutputLog::OnDockInLayoutClicked)
+			[
+				SNew(SHorizontalBox)
+				+ SHorizontalBox::Slot()
+				.AutoWidth()
+				.VAlign(VAlign_Center)
+				.Padding(4.0, 0.0f)
+				[
+					SNew(SImage)
+					.ColorAndOpacity(FSlateColor::UseForeground())
+					.Image(FAppStyle::Get().GetBrush("EditorViewport.SubMenu.Layouts"))
+				]
+				+ SHorizontalBox::Slot()
+				.VAlign(VAlign_Center)
+				.Padding(4.0, 0.0f)
+				[
+					SNew(STextBlock)
+					.Text(LOCTEXT("DockInLayout", "Dock in Layout"))
+					.ColorAndOpacity(FSlateColor::UseForeground())
+				]
+			];
+	}
+
+	return SNullWidget::NullWidget;
+}
+
 void SOutputLog::OpenLogFileInExplorer()
 {
 	FString Path = FPaths::ConvertRelativePathToFull(FPaths::ProjectLogDir());
@@ -1486,7 +1512,26 @@ void SOutputLog::OpenLogFileInExternalEditor()
 	}
 
 	FPlatformProcess::LaunchFileInDefaultExternalApplication(*Path, NULL, ELaunchVerb::Open);
-} 
+}
+
+FReply SOutputLog::OnDockInLayoutClicked()
+{
+	TSharedPtr<SDockTab> DockedTab;
+
+	static const FName OutputLogTabName = FName("OutputLog");
+	if (TSharedPtr<SDockTab> ActiveTab = FGlobalTabmanager::Get()->GetActiveTab())
+	{
+		DockedTab = ActiveTab->GetTabManager()->TryInvokeTab(OutputLogTabName);
+	}
+	
+	if (!DockedTab)
+	{
+		FGlobalTabmanager::Get()->TryInvokeTab(OutputLogTabName);
+	}
+
+	return FReply::Handled();
+}
+
 
 bool FOutputLogFilter::IsMessageAllowed(const TSharedPtr<FOutputLogMessage>& Message)
 {
diff --git a/Engine/Source/Developer/OutputLog/Private/SOutputLog.h b/Engine/Source/Developer/OutputLog/Private/SOutputLog.h
--- a/Engine/Source/Developer/OutputLog/Private/SOutputLog.h
+++ b/Engine/Source/Developer/OutputLog/Private/SOutputLog.h
@@ -313,6 +313,8 @@ public:
 		: _Messages()
 		{}
 		
+		SLATE_EVENT(FSimpleDelegate, OnCloseConsole)
+
 		/** All messages captured before this log window has been created */
 		SLATE_ARGUMENT( TArray< TSharedPtr<FOutputLogMessage> >, Messages )
 
@@ -326,7 +328,7 @@ public:
 	 *
 	 * @param	InArgs	Declaration used by the SNew() macro to construct this widget
 	 */
-	void Construct( const FArguments& InArgs );
+	void Construct( const FArguments& InArgs, bool bIsDrawerOutputLog );
 
 	// SWidget interface
 	virtual void Tick(const FGeometry& AllottedGeometry, const double InCurrentTime, const float InDeltaTime) override;
@@ -354,6 +356,8 @@ public:
 	 */
 	bool CanClearLog() const;
 
+	/** Focuses the edit box where you type in console commands */
+	void FocusConsoleCommandBox();
 protected:
 
 	virtual void Serialize( const TCHAR* V, ELogVerbosity::Type Verbosity, const class FName& Category ) override;
@@ -448,15 +452,20 @@ private:
 
 	TSharedRef<SWidget> GetViewButtonContent();
 
+	TSharedRef<SWidget> CreateDrawerDockButton();
+
 	void OpenLogFileInExplorer();
 
 	void OpenLogFileInExternalEditor();
 
-public:
+	FReply OnDockInLayoutClicked();
+protected:
+	TSharedPtr<SConsoleInputBox> ConsoleInputBox;
+
 	/** Visible messages filter */
 	FOutputLogFilter Filter;
 
-	TSharedPtr<class SComboButton> ViewOptionsComboButton;
+	bool bIsInDrawer = false;
 };
 
 /** Output log text marshaller to convert an array of FOutputLogMessages into styled lines to be consumed by an FTextLayout */
diff --git a/Engine/Source/Developer/OutputLog/Public/OutputLogModule.h b/Engine/Source/Developer/OutputLog/Public/OutputLogModule.h
--- a/Engine/Source/Developer/OutputLog/Public/OutputLogModule.h
+++ b/Engine/Source/Developer/OutputLog/Public/OutputLogModule.h
@@ -10,6 +10,10 @@
 #include "HAL/IConsoleManager.h"
 
 class SMultiLineEditableTextBox;
+class FOutputLogHistory;
+class SOutputLog;
+class SDockTab;
+class FSpawnTabArgs;
 
 /** Style of the debug console */
 namespace EDebugConsoleStyle
@@ -37,9 +41,13 @@ public:
 	virtual void StartupModule();
 	virtual void ShutdownModule();
 
+	static OUTPUTLOG_API FOutputLogModule& Get();
+
 	/** Generates a console input box widget.  Remember, this widget will become invalid if the
 		output log DLL is unloaded on the fly. */
-	virtual TSharedRef<SWidget> MakeConsoleInputBox(TSharedPtr<SMultiLineEditableTextBox>& OutExposedEditableTextBox, const FSimpleDelegate& OnCloseConsole) const;
+	virtual TSharedRef<SWidget> MakeConsoleInputBox(TSharedPtr<SMultiLineEditableTextBox>& OutExposedEditableTextBox, const FSimpleDelegate& OnCloseConsole, const FSimpleDelegate& OnConsoleCommandExecuted) const;
+
+	virtual TSharedRef<SWidget> MakeOutputLogDrawerWidget(const FSimpleDelegate& OnCloseConsole);
 
 	/** Opens a debug console in the specified window, if not already open */
 	virtual void ToggleDebugConsoleForWindow(const TSharedRef<SWindow>& Window, const EDebugConsoleStyle::Type InStyle, const FDebugConsoleDelegates& DebugConsoleDelegates);
@@ -49,7 +57,28 @@ public:
 
 	virtual void ClearOnPIE(const bool bIsSimulating);
 
+	virtual void FocusOutputLogConsoleBox(const TSharedRef<SWidget> OutputLogToFocus);
+
+	virtual const TSharedPtr<SWidget> GetOutputLog() const;
+	const TSharedPtr<SDockTab> GetOutputLogTab() const { return OutputLogTab.Pin(); }
+
+private:
+	TSharedRef<SDockTab> SpawnOutputLogTab(const FSpawnTabArgs& Args);
+	TSharedRef<SDockTab> SpawnDeviceOutputLogTab(const FSpawnTabArgs& Args);
+
 private:
+	/** Our global output log app spawner */
+	TSharedPtr<FOutputLogHistory> OutputLogHistory;
+
+	/** Our global active output log that belongs to a tab */
+	TWeakPtr<SOutputLog> OutputLog;
+
+	/** Global tab that the output log resides in */
+	TWeakPtr<SDockTab> OutputLogTab;
+
+	/** The output log that lives in a status bar drawer */
+	TWeakPtr<SOutputLog> OutputLogDrawer;
+
 	/** Weak pointer to a debug console that's currently open, if any */
 	TWeakPtr<SWidget> DebugConsole;
 
diff --git a/Engine/Source/Editor/EditorStyle/Private/EditorStyleClasses.cpp b/Engine/Source/Editor/EditorStyle/Private/EditorStyleClasses.cpp
--- a/Engine/Source/Editor/EditorStyle/Private/EditorStyleClasses.cpp
+++ b/Engine/Source/Editor/EditorStyle/Private/EditorStyleClasses.cpp
@@ -27,6 +27,8 @@ UEditorStyleSettings::UEditorStyleSettings( const FObjectInitializer& ObjectInit
 	
 	bUseGrid = true;
 
+	bCycleToOutputLogDrawer = true;
+
 	RegularColor = FLinearColor(0.035f, 0.035f, 0.035f);
 	RuleColor = FLinearColor(0.008f, 0.008f, 0.008f);
 	CenterColor = FLinearColor::Black;
diff --git a/Engine/Source/Editor/EditorStyle/Private/StarshipStyle.cpp b/Engine/Source/Editor/EditorStyle/Private/StarshipStyle.cpp
--- a/Engine/Source/Editor/EditorStyle/Private/StarshipStyle.cpp
+++ b/Engine/Source/Editor/EditorStyle/Private/StarshipStyle.cpp
@@ -48,7 +48,7 @@ void FStarshipEditorStyle::Initialize()
 	PropertyEditorModule.RegisterCustomPropertyTypeLayout("StyleColorList", FOnGetPropertyTypeCustomizationInstance::CreateStatic(&FStyleColorListCustomization::MakeInstance));
 
 	// Background
-	USlateThemeManager::Get().SetDefaultColor(EStyleColor::User1, USlateThemeManager::Get().GetColor(EStyleColor::Input));
+	USlateThemeManager::Get().SetDefaultColor(EStyleColor::User1, USlateThemeManager::Get().GetColor(EStyleColor::Recessed));
 	USlateThemeManager::Get().SetColorDisplayName(EStyleColor::User1, LOCTEXT("UserColor_OutputLogBackground", "Log Background"));
 	// Selection highlight
 	USlateThemeManager::Get().SetDefaultColor(EStyleColor::User2, USlateThemeManager::Get().GetColor(EStyleColor::Highlight));
@@ -981,7 +981,7 @@ void FStarshipEditorStyle::FStyle::SetupGeneralStyles()
 			.SetBackgroundImageHovered( BOX_BRUSH( "Common/WhiteGroupBorder", FMargin(4.0f/16.0f) ) )
 			.SetBackgroundImageFocused( BOX_BRUSH( "Common/WhiteGroupBorder", FMargin(4.0f/16.0f) ) )
 			.SetBackgroundImageReadOnly( BOX_BRUSH( "Common/WhiteGroupBorder", FMargin(4.0f/16.0f) ) )
-			.SetBackgroundColor( LogColor_Background )
+			.SetBackgroundColor(FStyleColors::Recessed)
 			);
 
 		Set("DebugConsole.Background", new FSlateNoResource());
diff --git a/Engine/Source/Editor/EditorStyle/Public/Classes/EditorStyleSettings.h b/Engine/Source/Editor/EditorStyle/Public/Classes/EditorStyleSettings.h
--- a/Engine/Source/Editor/EditorStyle/Public/Classes/EditorStyleSettings.h
+++ b/Engine/Source/Editor/EditorStyle/Public/Classes/EditorStyleSettings.h
@@ -149,6 +149,13 @@ public:
 	UPROPERTY(EditAnywhere, config, Category="Output Log")
 	bool bPromoteOutputLogWarningsDuringPIE;
 
+	/**
+	 * If checked pressing the console command shortcut will cycle between focusing the status bar console, opening the output log drawer, and back to the previous focus target. 
+	 * If unchecked, the console command shortcut will only focus the status bar console
+	 */
+	UPROPERTY(EditAnywhere, config, Category = "Output Log", meta = (DisplayName = "Open Output Log Drawer with Console Command Shortcut"))
+	bool bCycleToOutputLogDrawer;
+
 	/** New asset editor tabs will open at the specified location. */
 	UPROPERTY(EditAnywhere, config, Category=UserInterface)
 	EAssetEditorOpenLocation AssetEditorOpenLocation;
diff --git a/Engine/Source/Editor/StatusBar/Private/SStatusBar.cpp b/Engine/Source/Editor/StatusBar/Private/SStatusBar.cpp
--- a/Engine/Source/Editor/StatusBar/Private/SStatusBar.cpp
+++ b/Engine/Source/Editor/StatusBar/Private/SStatusBar.cpp
@@ -11,7 +11,6 @@
 #include "Widgets/Layout/SBox.h"
 #include "Widgets/Layout/SBorder.h"
 #include "Widgets/Input/SButton.h"
-#include "Widgets/Input/SMultiLineEditableTextBox.h"
 #include "Widgets/Layout/SSpacer.h"
 #include "Widgets/Docking/SDockTab.h"
 #include "Widgets/Layout/SSeparator.h"
@@ -19,7 +18,6 @@
 #include "Framework/Commands/Commands.h"
 #include "ToolMenuContext.h"
 #include "ToolMenus.h"
-#include "OutputLogModule.h"
 #include "SourceControlMenuHelpers.h"
 #include "Framework/Application/SlateApplication.h"
 #include "Widgets/Layout/SSplitter.h"
@@ -55,14 +53,13 @@ class SDrawerOverlay : public SCompoundWidget
 		SLATE_ARGUMENT(float, MaxDrawerHeight)
 		SLATE_ARGUMENT(float, TargetDrawerHeight)
 		SLATE_EVENT(FOnStatusBarDrawerTargetHeightChanged, OnTargetHeightChanged)
+		SLATE_EVENT(FSimpleDelegate, OnDismissComplete)
 		SLATE_ARGUMENT(FVector2D, ShadowOffset)
 	SLATE_END_ARGS()
 
-	void UpdateHeightInterp(float InAlpha)
+	~SDrawerOverlay()
 	{
-		float NewHeight = FMath::Lerp(0.0f, TargetHeight, InAlpha);
-
-		SetHeight(NewHeight);
+		FSlateThrottleManager::Get().LeaveResponsiveMode(AnimationThrottle);
 	}
 
 	void Construct(const FArguments& InArgs)
@@ -89,12 +86,23 @@ class SDrawerOverlay : public SCompoundWidget
 		bIsResizeHandleHovered = false;
 		bIsResizing = false;
 
+		OnDismissComplete = InArgs._OnDismissComplete;
+
+		DrawerEasingCurve = FCurveSequence(0.0f, 0.15f, ECurveEaseFunction::QuadOut);
+
 		ChildSlot
 		[
 			InArgs._Content.Widget
 		];
 	}
 
+	void UpdateHeightInterp(float InAlpha)
+	{
+		float NewHeight = FMath::Lerp(0.0f, TargetHeight, InAlpha);
+
+		SetHeight(NewHeight);
+	}
+
 	virtual bool SupportsKeyboardFocus() const override
 	{
 		return true;
@@ -253,6 +261,30 @@ class SDrawerOverlay : public SCompoundWidget
 
 	}
 
+	void Open()
+	{
+		DrawerEasingCurve.Play(AsShared(), false, DrawerEasingCurve.IsPlaying() ? DrawerEasingCurve.GetSequenceTime() : 0.0f, false);
+
+		if (!DrawerOpenCloseTimer.IsValid())
+		{
+			AnimationThrottle = FSlateThrottleManager::Get().EnterResponsiveMode();
+			DrawerOpenCloseTimer = RegisterActiveTimer(0.0f, FWidgetActiveTimerDelegate::CreateSP(this, &SDrawerOverlay::UpdateDrawerAnimation));
+		}
+	}
+
+	void Dismiss()
+	{
+		if (DrawerEasingCurve.IsForward())
+		{
+			DrawerEasingCurve.Reverse();
+		}
+
+		if (!DrawerOpenCloseTimer.IsValid())
+		{
+			AnimationThrottle = FSlateThrottleManager::Get().EnterResponsiveMode();
+			DrawerOpenCloseTimer = RegisterActiveTimer(0.0f, FWidgetActiveTimerDelegate::CreateSP(this, &SDrawerOverlay::UpdateDrawerAnimation));
+		}
+	}
 private:
 	FGeometry GetRenderTransformedGeometry(const FGeometry& AllottedGeometry) const
 	{
@@ -268,14 +300,39 @@ private:
 	{
 		CurrentHeight = FMath::Clamp(NewHeight, MinHeight, TargetHeight);
 	}
+
+	EActiveTimerReturnType UpdateDrawerAnimation(double CurrentTime, float DeltaTime)
+	{
+		UpdateHeightInterp(DrawerEasingCurve.GetLerp());
+
+		if (!DrawerEasingCurve.IsPlaying())
+		{
+			if (DrawerEasingCurve.IsAtStart())
+			{
+				OnDismissComplete.ExecuteIfBound();
+			}
+
+			FSlateThrottleManager::Get().LeaveResponsiveMode(AnimationThrottle);
+			DrawerOpenCloseTimer.Reset();
+			return EActiveTimerReturnType::Stop;
+		}
+
+		return EActiveTimerReturnType::Continue;
+	}
+
+
 private:
 	FGeometry InitialResizeGeometry;
+	TSharedPtr<FActiveTimerHandle> DrawerOpenCloseTimer;
 	FOnStatusBarDrawerTargetHeightChanged OnTargetHeightChanged;
+	FCurveSequence DrawerEasingCurve;
+	FSimpleDelegate OnDismissComplete;
 	const FSlateBrush* BackgroundBrush;
 	const FSlateBrush* ShadowBrush;
 	const FSlateBrush* BorderBrush;
 	const FSplitterStyle* SplitterStyle;
 	FVector2D ShadowOffset;
+	FThrottleRequest AnimationThrottle;
 	FThrottleRequest ResizeThrottleHandle;
 	float ExpanderSize;
 	float CurrentHeight;
@@ -527,7 +584,6 @@ void SStatusBar::Construct(const FArguments& InArgs, FName InStatusBarName, cons
 
 	const FSlateBrush* StatusBarBackground = FAppStyle::Get().GetBrush("Brushes.Panel");
 
-	DrawerEasingCurve = FCurveSequence(0.0f, 0.15f, ECurveEaseFunction::QuadOut);
 
 	FSlateApplication::Get().OnFocusChanging().AddSP(this, &SStatusBar::OnGlobalFocusChanging);
 	FGlobalTabmanager::Get()->OnActiveTabChanged_Subscribe(FOnActiveTabChanged::FDelegate::CreateSP(this, &SStatusBar::OnActiveTabChanged));
@@ -547,18 +603,6 @@ void SStatusBar::Construct(const FArguments& InArgs, FName InStatusBarName, cons
 			+SHorizontalBox::Slot()
 			[
 				SNew(SHorizontalBox)
-				 +SHorizontalBox::Slot()
-			    .AutoWidth()
-			    .Padding(1.0f, 0.0f)
-			    [
-				    SNew(SBorder)
-				    .BorderImage(StatusBarBackground)
-				    .VAlign(VAlign_Center)
-				    .Padding(FMargin(6.0f, 0.0f))
-				    [
-					    MakeDebugConsoleWidget(InArgs._OnConsoleClosed)
-				    ]
-			    ]
 				+SHorizontalBox::Slot()
 				.FillWidth(1.0f)
 				.Padding(1.0f, 0.0f)
@@ -708,16 +752,6 @@ TSharedPtr<SDockTab> SStatusBar::GetParentTab() const
 	return ParentTab.Pin();
 }
 
-bool SStatusBar::FocusDebugConsole()
-{
-	return FSlateApplication::Get().SetKeyboardFocus(ConsoleEditBox, EFocusCause::SetDirectly);
-}
-
-bool SStatusBar::IsDebugConsoleFocused() const
-{
-	return ConsoleEditBox->HasKeyboardFocus();
-}
-
 void SStatusBar::OnGlobalFocusChanging(const FFocusEvent& FocusEvent, const FWeakWidgetPath& OldFocusedWidgetPath, const TSharedPtr<SWidget>& OldFocusedWidget, const FWidgetPath& NewFocusedWidgetPath, const TSharedPtr<SWidget>& NewFocusedWidget)
 {
 	// Sometimes when dismissing focus can change which will trigger this again
@@ -730,9 +764,9 @@ void SStatusBar::OnGlobalFocusChanging(const FFocusEvent& FocusEvent, const FWea
 		TSharedRef<SWidget> ThisWidget = AsShared();
 
 		TSharedPtr<SWidget> ActiveDrawerOverlayContent;
-		if (!OpenedDrawerData.Key.IsNone())
+		if (OpenedDrawer.IsValid())
 		{
-			ActiveDrawerOverlayContent = OpenedDrawerData.Value;
+			ActiveDrawerOverlayContent = OpenedDrawer.DrawerOverlay;
 		}
 
 		bool bShouldDismiss = false;
@@ -808,7 +842,7 @@ TSharedRef<SWidget> SStatusBar::MakeStatusBarDrawerButton(const FStatusBarDrawer
 {
 	const FName DrawerId = Drawer.UniqueId;
 
-	return
+	TSharedRef<SWidget> DrawerButton = 
 		SNew(SButton)
 		.IsFocusable(false)
 		.ButtonStyle(&FAppStyle::Get().GetWidgetStyle<FButtonStyle>("StatusBar.StatusBarButton"))
@@ -824,7 +858,7 @@ TSharedRef<SWidget> SStatusBar::MakeStatusBarDrawerButton(const FStatusBarDrawer
 			[
 				SNew(SImage)
 				.ColorAndOpacity(FSlateColor::UseForeground())
-				.Image_Lambda([this, DrawerId](){ return IsDrawerOpened(DrawerId) ? DownArrow : UpArrow; })
+				.Image_Lambda([this, DrawerId]() { return IsDrawerOpened(DrawerId) ? DownArrow : UpArrow; })
 			]
 			+ SHorizontalBox::Slot()
 			.Padding(2.0f)
@@ -844,7 +878,29 @@ TSharedRef<SWidget> SStatusBar::MakeStatusBarDrawerButton(const FStatusBarDrawer
 				.TextStyle(&FAppStyle::Get().GetWidgetStyle<FTextBlockStyle>("NormalText"))
 				.Text(Drawer.ButtonText)
 			]
-		];	
+		];
+
+
+	if (Drawer.CustomWidget)
+	{
+		return
+			SNew(SHorizontalBox)
+			+ SHorizontalBox::Slot()
+			.Padding(0.0f, 0.0f, 4.0f, 0.0f)
+			.AutoWidth()
+			[
+				DrawerButton
+			]
+			+ SHorizontalBox::Slot()
+			[
+				Drawer.CustomWidget.ToSharedRef()
+			];
+	
+	}
+	else
+	{
+		return DrawerButton;
+	}
 }
 
 TSharedRef<SWidget> SStatusBar::MakeStatusBarToolBarWidget()
@@ -857,18 +913,6 @@ TSharedRef<SWidget> SStatusBar::MakeStatusBarToolBarWidget()
 	return UToolMenus::Get()->GenerateWidget(StatusBarToolBarName, MenuContext);
 }
 
-TSharedRef<SWidget> SStatusBar::MakeDebugConsoleWidget(FSimpleDelegate OnConsoleClosed)
-{
-	FOutputLogModule& OutputLogModule = FModuleManager::LoadModuleChecked<FOutputLogModule>(TEXT("OutputLog"));
-
-	return
-		SNew(SBox)
-		.WidthOverride(350.f)
-		[
-			OutputLogModule.MakeConsoleInputBox(ConsoleEditBox, OnConsoleClosed)
-		];
-}
-
 TSharedRef<SWidget> SStatusBar::MakeStatusMessageWidget()
 {
 	return 
@@ -904,7 +948,12 @@ TSharedRef<SWidget> SStatusBar::MakeProgressBar()
 
 bool SStatusBar::IsDrawerOpened(const FName DrawerId) const
 {
-	return OpenedDrawerData.Key == DrawerId ? true : false;
+	return OpenedDrawer == DrawerId ? true : false;
+}
+
+bool SStatusBar::IsAnyOtherDrawerOpened(const FName DrawerId) const
+{
+	return OpenedDrawer.IsValid() && OpenedDrawer.DrawerId != DrawerId ? true : false;
 }
 
 FReply SStatusBar::OnDrawerButtonClicked(const FName DrawerId)
@@ -922,35 +971,15 @@ FReply SStatusBar::OnDrawerButtonClicked(const FName DrawerId)
 }
 
 
-EActiveTimerReturnType SStatusBar::UpdateDrawerAnimation(double CurrentTime, float DeltaTime)
-{
-	check(OpenedDrawerData.Value.IsValid());
-
-	OpenedDrawerData.Value->UpdateHeightInterp(DrawerEasingCurve.GetLerp());
-
-	if (!DrawerEasingCurve.IsPlaying())
-	{
-		if(DrawerEasingCurve.IsAtStart())
-		{
-			CloseDrawerImmediately();
-		}
-
-		FSlateThrottleManager::Get().LeaveResponsiveMode(AnimationThrottle);
-		DrawerOpenCloseTimer.Reset();
-		return EActiveTimerReturnType::Stop;
-	}
-
-	return EActiveTimerReturnType::Continue;
-}
 
 void SStatusBar::OnDrawerHeightChanged(float TargetHeight)
 {
-	TSharedPtr<SWindow> MyWindow = WindowWithOverlayContent.Pin();
+	TSharedPtr<SWindow> MyWindow = OpenedDrawer.WindowWithOverlayContent.Pin();
 
 	// Save the height has a percentage of the screen
 	const float TargetDrawerHeightPct = TargetHeight / (MyWindow->GetSizeInScreen().Y / MyWindow->GetDPIScaleFactor());
 
-	GConfig->SetFloat(TEXT("DrawerSizes"), *(StatusBarName.ToString() + TEXT(".") + OpenedDrawerData.Key.ToString()), TargetDrawerHeightPct, GEditorSettingsIni);
+	GConfig->SetFloat(TEXT("DrawerSizes"), *(StatusBarName.ToString() + TEXT(".") + OpenedDrawer.DrawerId.ToString()), TargetDrawerHeightPct, GEditorSettingsIni);
 }
 
 void SStatusBar::RegisterStatusBarMenu()
@@ -1135,6 +1164,20 @@ TSharedRef<SWidget> SStatusBar::OnGetProgressBarMenuContent()
 		];
 }
 
+void SStatusBar::CloseDrawerImmediatelyInternal(const FOpenDrawerData& Data)
+{
+	if (Data.IsValid())
+	{
+		TSharedRef<SWidget> DrawerOverlayContent = Data.DrawerOverlay.ToSharedRef();
+
+		// Remove the content browser from the window
+		if (TSharedPtr<SWindow> Window = Data.WindowWithOverlayContent.Pin())
+		{
+			Window->RemoveOverlaySlot(DrawerOverlayContent);
+		}
+	}
+}
+
 void SStatusBar::RegisterDrawer(FStatusBarDrawer&& Drawer)
 {
 	const int32 NumDrawers = RegisteredDrawers.Num();
@@ -1162,7 +1205,7 @@ void SStatusBar::RegisterDrawer(FStatusBarDrawer&& Drawer)
 void SStatusBar::OpenDrawer(const FName DrawerId)
 {
 	// Close any other open drawer
-	if (DrawerId != OpenedDrawerData.Key)
+	if (OpenedDrawer.DrawerId != DrawerId && DismissingDrawers.IndexOfByKey(DrawerId) == INDEX_NONE)
 	{
 		DismissDrawer(nullptr);
 
@@ -1182,31 +1225,34 @@ void SStatusBar::OpenDrawer(const FName DrawerId)
 			float TargetDrawerHeight = (MyWindow->GetSizeInScreen().Y * TargetDrawerHeightPct) / MyWindow->GetDPIScaleFactor();
 
 			const float MinDrawerHeight = GetTickSpaceGeometry().GetLocalSize().Y + MyWindow->GetWindowBorderSize().Bottom;
-			DrawerEasingCurve.Play(ThisStatusBar, false, DrawerEasingCurve.IsPlaying() ? DrawerEasingCurve.GetSequenceTime() : 0.0f, false);
+	
+			FOpenDrawerData NewlyOpenedDrawer;
 
 			MyWindow->AddOverlaySlot()
 				.VAlign(VAlign_Bottom)
 				.Padding(FMargin(10.0f, 20.0f, 10.0f, MinDrawerHeight))
 				[
-					SAssignNew(OpenedDrawerData.Value, SDrawerOverlay)
+					SAssignNew(NewlyOpenedDrawer.DrawerOverlay, SDrawerOverlay)
 					.MinDrawerHeight(MinDrawerHeight)
 					.TargetDrawerHeight(TargetDrawerHeight)
 					.MaxDrawerHeight(MaxDrawerHeight)
+					.OnDismissComplete_Lambda(
+						[DrawerId, this]()
+						{
+							CloseDrawerImmediately(DrawerId);
+						})
 					.OnTargetHeightChanged(this, &SStatusBar::OnDrawerHeightChanged)
 					[
 						DrawerData->GetDrawerContentDelegate.Execute()
 					]
 				];
 
-			WindowWithOverlayContent = MyWindow;
+			NewlyOpenedDrawer.WindowWithOverlayContent = MyWindow;
+			NewlyOpenedDrawer.DrawerId = DrawerId;
+			NewlyOpenedDrawer.DrawerOverlay->Open();
 
-			if (!DrawerOpenCloseTimer.IsValid())
-			{
-				AnimationThrottle = FSlateThrottleManager::Get().EnterResponsiveMode();
-				DrawerOpenCloseTimer = RegisterActiveTimer(0.0f, FWidgetActiveTimerDelegate::CreateSP(ThisStatusBar, &SStatusBar::UpdateDrawerAnimation));
-			}
+			OpenedDrawer = MoveTemp(NewlyOpenedDrawer);
 
-			OpenedDrawerData.Key = DrawerId;
 			DrawerData->OnDrawerOpenedDelegate.ExecuteIfBound(ThisStatusBar);
 		}
 	}
@@ -1214,50 +1260,48 @@ void SStatusBar::OpenDrawer(const FName DrawerId)
 
 void SStatusBar::DismissDrawer(const TSharedPtr<SWidget>& NewlyFocusedWidget)
 {
-	if (!OpenedDrawerData.Key.IsNone())
+	if (OpenedDrawer.IsValid())
 	{
-		FStatusBarDrawer* Drawer = RegisteredDrawers.FindByKey(OpenedDrawerData.Key);
+		FStatusBarDrawer* Drawer = RegisteredDrawers.FindByKey(OpenedDrawer.DrawerId);
 
-		if (DrawerEasingCurve.IsForward())
-		{
-			DrawerEasingCurve.Reverse();
-		}
+		OpenedDrawer.DrawerOverlay->Dismiss();
+		DismissingDrawers.Add(MoveTemp(OpenedDrawer));
 
-		if (!DrawerOpenCloseTimer.IsValid())
-		{
-			AnimationThrottle = FSlateThrottleManager::Get().EnterResponsiveMode();
-			DrawerOpenCloseTimer = RegisterActiveTimer(0.0f, FWidgetActiveTimerDelegate::CreateSP(this, &SStatusBar::UpdateDrawerAnimation));
-		}
+		OpenedDrawer = FOpenDrawerData();
 
 		Drawer->OnDrawerDismissedDelegate.ExecuteIfBound(NewlyFocusedWidget);
 	}
 }
 
-void SStatusBar::CloseDrawerImmediately()
+void SStatusBar::CloseDrawerImmediately(FName DrawerId)
 {
-	if (!OpenedDrawerData.Key.IsNone())
+	// If no ID is specified remove all drawers
+	if (DrawerId.IsNone())
 	{
-		TSharedRef<SWidget> DrawerOverlayContent = OpenedDrawerData.Value.ToSharedRef();
-
-		// Remove the content browser from the window
-		if (TSharedPtr<SWindow> Window = WindowWithOverlayContent.Pin())
+		for (const FOpenDrawerData& Data : DismissingDrawers)
 		{
-			Window->RemoveOverlaySlot(DrawerOverlayContent);
+			CloseDrawerImmediatelyInternal(Data);
 		}
 
-		OpenedDrawerData = TPair<FName, TSharedPtr<SDrawerOverlay>>();
+		DismissingDrawers.Empty();
 
-		WindowWithOverlayContent.Reset();
+		CloseDrawerImmediatelyInternal(OpenedDrawer);
 
-		if (DrawerOpenCloseTimer.IsValid())
+		OpenedDrawer = FOpenDrawerData();
+	}
+	else
+	{
+		int32 Index = DismissingDrawers.IndexOfByKey(DrawerId);
+		if (Index != INDEX_NONE)
 		{
-			FSlateThrottleManager::Get().LeaveResponsiveMode(AnimationThrottle);
-			UnRegisterActiveTimer(DrawerOpenCloseTimer.ToSharedRef());
+			CloseDrawerImmediatelyInternal(DismissingDrawers[Index]);
+			DismissingDrawers.RemoveAtSwap(Index);
+		}
+		else if (OpenedDrawer == DrawerId)
+		{
+			CloseDrawerImmediatelyInternal(OpenedDrawer);
+			OpenedDrawer = FOpenDrawerData();
 		}
-
-		DrawerOpenCloseTimer.Reset();
-		DrawerEasingCurve.JumpToStart();
-		DrawerEasingCurve.Pause();
 	}
 }
 
diff --git a/Engine/Source/Editor/StatusBar/Private/SStatusBar.h b/Engine/Source/Editor/StatusBar/Private/SStatusBar.h
--- a/Engine/Source/Editor/StatusBar/Private/SStatusBar.h
+++ b/Engine/Source/Editor/StatusBar/Private/SStatusBar.h
@@ -26,6 +26,7 @@ class SStatusBarProgressWidget;
 namespace StatusBarDrawerIds
 {
 	const FName ContentBrowser("ContentBrowser");
+	const FName OutputLog("OutputLog");
 }
 
 /** Data payload for messages in the status bar */
@@ -71,6 +72,7 @@ struct FStatusBarDrawer
 	FOnStatusBarDrawerOpened OnDrawerOpenedDelegate;
 	FOnStatusBarDrawerDismissed OnDrawerDismissedDelegate;
 
+	TSharedPtr<SWidget> CustomWidget;
 	FText ButtonText;
 	FText ToolTipText;
 	const FSlateBrush* Icon;
@@ -86,11 +88,27 @@ struct FStatusBarDrawer
 	}
 };
 
+struct FOpenDrawerData
+{
+	FName DrawerId;
+	TSharedPtr<SDrawerOverlay> DrawerOverlay;
+	TWeakPtr<SWindow> WindowWithOverlayContent;
+
+	bool IsValid() const
+	{
+		return !DrawerId.IsNone();
+	}
+
+	bool operator==(const FName InDrawerId) const
+	{
+		return DrawerId == InDrawerId;
+	}
+};
+
 class SStatusBar : public SCompoundWidget
 {
 	SLATE_BEGIN_ARGS(SStatusBar)	
 	{}
-		SLATE_EVENT(FSimpleDelegate, OnConsoleClosed)
 	SLATE_END_ARGS()
 
 public:
@@ -143,12 +161,21 @@ public:
 	 */
 	bool CancelProgressNotification(FProgressNotificationHandle InHandle);
 
+	/**
+	 * Gets the owning major tab that this status bar is in
+	 */
 	TSharedPtr<SDockTab> GetParentTab() const;
-	bool FocusDebugConsole();
-	bool IsDebugConsoleFocused() const;
 
+	/**
+	 * Registers a new drawer with this status bar. Registering will add a button to open and close the drawer
+	 */
 	void RegisterDrawer(FStatusBarDrawer&& Drawer);
 
+	/**
+	 * Opens a drawer
+	 * 
+	 * @param DrawerId Name of the registered drawer to open 
+	 */
 	void OpenDrawer(const FName DrawerId);
 
 	/**
@@ -156,9 +183,18 @@ public:
 	 */
 	void DismissDrawer(const TSharedPtr<SWidget>& NewlyFocusedWidget);
 
-	void CloseDrawerImmediately();
+	/**
+	 * Closes a drawer immediately with no closing animation. Needed for when UI is shutting down or to prevent conflicts where a drawer is open in two places at once. 
+	 * @param DrawerId The name of the drawer to close. If no name is specified all drawers are closed immediately
+	 */
+	void CloseDrawerImmediately(FName DrawerId = NAME_None);
 
+	/** Is a specific registered drawer currently open */
 	bool IsDrawerOpened(const FName DrawerId) const;
+
+	/** Is any drawer other than the one specified by DrawerId opened */
+	bool IsAnyOtherDrawerOpened(const FName DrawerId) const;
+
 private:
 	/** Called when global focus changes which is used to determine if we should close an opened content browser drawer */
 	void OnGlobalFocusChanging(const FFocusEvent& FocusEvent, const FWeakWidgetPath& OldFocusedWidgetPath, const TSharedPtr<SWidget>& OldFocusedWidget, const FWidgetPath& NewFocusedWidgetPath, const TSharedPtr<SWidget>& NewFocusedWidget);
@@ -170,25 +206,15 @@ private:
 
 	EVisibility GetHelpIconVisibility() const;
 
-	bool GetContentBrowserExpanded() const;
-
 	FText GetStatusBarMessage() const;
 
 	TSharedRef<SWidget> MakeStatusBarDrawerButton(const FStatusBarDrawer& Drawer);
-
 	TSharedRef<SWidget> MakeStatusBarToolBarWidget();
-	TSharedRef<SWidget> MakeDebugConsoleWidget(FSimpleDelegate OnConsoleClosed);
 	TSharedRef<SWidget> MakeStatusMessageWidget();
 	TSharedRef<SWidget> MakeProgressBar();
 
-	void MakeDrawerButtons(TSharedRef<SHorizontalBox> DrawerBox);
-
 	FReply OnDrawerButtonClicked(const FName DrawerId);
 
-	void CreateAnimationTimerIfNeeded();
-
-	EActiveTimerReturnType UpdateDrawerAnimation(double CurrentTime, float DeltaTime);
-
 	void RegisterStatusBarMenu();
 
 	void RegisterSourceControlStatus();
@@ -201,32 +227,27 @@ private:
 	
 	TSharedRef<SWidget> OnGetProgressBarMenuContent();
 
+	void CloseDrawerImmediatelyInternal(const FOpenDrawerData& Data);
 private:
 	TArray<FStatusBarMessage> MessageStack;
 	TArray<FStatusBarProgress> ProgressNotifications;
 
-	TSharedPtr<SMultiLineEditableTextBox> ConsoleEditBox;
 	TWeakPtr<SDockTab> ParentTab;
 
 	TArray<FStatusBarDrawer> RegisteredDrawers;
 
-	TWeakPtr<SWindow> WindowWithOverlayContent;
+	FOpenDrawerData OpenedDrawer;
 
-	TPair<FName,TSharedPtr<SDrawerOverlay>> OpenedDrawerData;
+	TArray<FOpenDrawerData> DismissingDrawers;
 
 	TSharedPtr<SHorizontalBox> DrawerBox;
 
-	TSharedPtr<FActiveTimerHandle> DrawerOpenCloseTimer;
-
 	TSharedPtr<SStatusBarProgressArea> ProgressBar;
 
 	TWeakPtr<SNotificationItem> ActiveProgressNotification;
 
 	TWeakPtr<SStatusBarProgressWidget> ActiveNotificationProgressWidget;
 
-
-	FCurveSequence DrawerEasingCurve;
-	FThrottleRequest AnimationThrottle;
 	const FSlateBrush* UpArrow = nullptr;
 	const FSlateBrush* DownArrow = nullptr;
 	FName StatusBarName;
diff --git a/Engine/Source/Editor/StatusBar/Private/StatusBarSubsystem.cpp b/Engine/Source/Editor/StatusBar/Private/StatusBarSubsystem.cpp
--- a/Engine/Source/Editor/StatusBar/Private/StatusBarSubsystem.cpp
+++ b/Engine/Source/Editor/StatusBar/Private/StatusBarSubsystem.cpp
@@ -21,6 +21,9 @@
 #include "Widgets/Images/SImage.h"
 #include "Widgets/Input/SButton.h"
 #include "Misc/ConfigCacheIni.h"
+#include "OutputLogModule.h"
+#include "Widgets/Input/SMultiLineEditableTextBox.h"
+#include "Classes/EditorStyleSettings.h"
 
 #define LOCTEXT_NAMESPACE "StatusBar"
 
@@ -217,28 +220,74 @@ void UStatusBarSubsystem::Deinitialize()
 	FSlateNotificationManager::Get().SetProgressNotificationHandler(nullptr);
 
 	StatusBarContentBrowser.Reset();
+	StatusBarOutputLog.Reset();
 }
 
-bool UStatusBarSubsystem::FocusDebugConsole(TSharedRef<SWindow> ParentWindow)
+bool UStatusBarSubsystem::ToggleDebugConsole(TSharedRef<SWindow> ParentWindow, bool bAlwaysToggleDrawer)
 {
-	bool bFocusedSuccessfully = false;
+	bool bToggledSuccessfully = false;
+
+	FOutputLogModule& OutputLogModule = FOutputLogModule::Get();
+
+	// Get the global output log tab if it exists. If it exists and is in the same editor as the status bar we'll focus that instead
+	TSharedPtr<SDockTab> MainOutputLogTab = OutputLogModule.GetOutputLogTab();
+	
+	const bool bCycleToOutputLogDrawer = GetDefault<UEditorStyleSettings>()->bCycleToOutputLogDrawer || bAlwaysToggleDrawer;
 
 	for (auto StatusBar : StatusBars)
 	{
-		if (TSharedPtr<SStatusBar> StatusBarPinned = StatusBar.Value.Pin())
+		FStatusBarData& SBData = StatusBar.Value;
+		if (TSharedPtr<SStatusBar> StatusBarPinned = SBData.StatusBarWidget.Pin())
 		{
 			TSharedPtr<SDockTab> ParentTab = StatusBarPinned->GetParentTab();
 			if (ParentTab && ParentTab->IsForeground() && ParentTab->GetParentWindow() == ParentWindow)
 			{
-				// Cache off the previously focused widget so we can restore focus if the user hits the focus key again
-				PreviousKeyboardFocusedWidget = FSlateApplication::Get().GetKeyboardFocusedWidget();
-				bFocusedSuccessfully = StatusBarPinned->FocusDebugConsole();
+				if (!bAlwaysToggleDrawer && MainOutputLogTab && MainOutputLogTab->GetTabManager()->GetOwnerTab() == ParentTab)
+				{
+					OutputLogModule.FocusOutputLogConsoleBox(OutputLogModule.GetOutputLog().ToSharedRef());
+				}
+				else
+				{
+					// This toggles between 3 states: 
+					// If the drawer is opened, close it
+					// if the console edit box is focused, open the drawer
+					// if something else is focused, focus the console edit box
+					if (StatusBarPinned->IsDrawerOpened(StatusBarDrawerIds::OutputLog))
+					{
+						StatusBarPinned->DismissDrawer(nullptr);
+					}
+					else if (SBData.ConsoleEditBox->HasKeyboardFocus() || bAlwaysToggleDrawer)
+					{
+						if (bCycleToOutputLogDrawer)
+						{
+							StatusBarPinned->OpenDrawer(StatusBarDrawerIds::OutputLog);
+						}
+						else
+						{
+							// Restore previous focus
+							FSlateApplication::Get().SetKeyboardFocus(PreviousKeyboardFocusedWidget.Pin(), EFocusCause::SetDirectly);
+
+							PreviousKeyboardFocusedWidget.Reset();
+						}
+					}
+					else
+					{
+
+						// Cache off the previously focused widget so we can restore focus if the user hits the focus key again
+						PreviousKeyboardFocusedWidget = FSlateApplication::Get().GetKeyboardFocusedWidget();
+
+						TSharedPtr<SWidget> WidgetToFocus = SBData.ConsoleEditBox;
+						FSlateApplication::Get().SetKeyboardFocus(WidgetToFocus, EFocusCause::SetDirectly);
+					}
+				}
+
+				bToggledSuccessfully = true;
 				break;
 			}
 		}
 	}
 
-	return bFocusedSuccessfully;
+	return bToggledSuccessfully;
 }
 
 bool UStatusBarSubsystem::OpenContentBrowserDrawer()
@@ -264,12 +313,34 @@ bool UStatusBarSubsystem::OpenContentBrowserDrawer()
 	return false;
 }
 
+bool UStatusBarSubsystem::OpenOutputLogDrawer()
+{
+	TSharedPtr<SWindow> ParentWindow = FSlateApplication::Get().GetActiveTopLevelWindow();
+	if (!ParentWindow.IsValid())
+	{
+		if (TSharedPtr<SDockTab> ActiveTab = FGlobalTabmanager::Get()->GetActiveTab())
+		{
+			if (TSharedPtr<SDockTab> ActiveMajorTab = FGlobalTabmanager::Get()->GetMajorTabForTabManager(ActiveTab->GetTabManager()))
+			{
+				ParentWindow = ActiveMajorTab->GetParentWindow();
+			}
+		}
+	}
+
+	if (ParentWindow.IsValid() && ParentWindow->GetType() == EWindowType::Normal)
+	{
+		return ToggleDebugConsole(ParentWindow.ToSharedRef(), true);
+	}
+
+	return false;
+}
+
 bool UStatusBarSubsystem::ForceDismissDrawer()
 {
 	bool bWasDismissed = false;
 	for (auto StatusBar : StatusBars)
 	{
-		if (TSharedPtr<SStatusBar> StatusBarPinned = StatusBar.Value.Pin())
+		if (TSharedPtr<SStatusBar> StatusBarPinned = StatusBar.Value.StatusBarWidget.Pin())
 		{
 			if (StatusBarPinned->IsDrawerOpened(StatusBarDrawerIds::ContentBrowser))
 			{
@@ -290,7 +361,7 @@ bool UStatusBarSubsystem::ToggleContentBrowser(TSharedRef<SWindow> ParentWindow)
 
 	for (auto StatusBar : StatusBars)
 	{
-		if (TSharedPtr<SStatusBar> StatusBarPinned = StatusBar.Value.Pin())
+		if (TSharedPtr<SStatusBar> StatusBarPinned = StatusBar.Value.StatusBarWidget.Pin())
 		{
 			if(StatusBarPinned->IsDrawerOpened(StatusBarDrawerIds::ContentBrowser))
 			{
@@ -318,8 +389,7 @@ TSharedRef<SWidget> UStatusBarSubsystem::MakeStatusBarWidget(FName StatusBarName
 	CreateContentBrowserIfNeeded();
 
 	TSharedRef<SStatusBar> StatusBar =
-		SNew(SStatusBar, StatusBarName, InParentTab)
-		.OnConsoleClosed_UObject(this, &UStatusBarSubsystem::OnDebugConsoleClosed);
+		SNew(SStatusBar, StatusBarName, InParentTab);
 
 	FStatusBarDrawer ContentBrowserDrawer(StatusBarDrawerIds::ContentBrowser);
 	ContentBrowserDrawer.GetDrawerContentDelegate.BindUObject(this, &UStatusBarSubsystem::OnGetContentBrowser);
@@ -331,16 +401,54 @@ TSharedRef<SWidget> UStatusBarSubsystem::MakeStatusBarWidget(FName StatusBarName
 
 	StatusBar->RegisterDrawer(MoveTemp(ContentBrowserDrawer));
 
+	FOutputLogModule& OutputLogModule = FOutputLogModule::Get();
+
+	TWeakPtr<SStatusBar> StatusBarWeakPtr = StatusBar;
+
+	FSimpleDelegate OnConsoleClosed = FSimpleDelegate::CreateUObject(this, &UStatusBarSubsystem::OnDebugConsoleClosed, StatusBarWeakPtr);
+	FSimpleDelegate OnConsoleCommandExecuted;
+
+	TSharedPtr<SMultiLineEditableTextBox> ConsoleEditBox;
+	TSharedRef<SWidget> OutputLog = 
+			SNew(SBorder)
+			.BorderImage(FAppStyle::Get().GetBrush("Brushes.Panel"))
+			.VAlign(VAlign_Center)
+			.Padding(FMargin(6.0f, 0.0f))
+			[
+				SNew(SBox)
+				.WidthOverride(350.f)
+				[
+					OutputLogModule.MakeConsoleInputBox(ConsoleEditBox, OnConsoleClosed, OnConsoleCommandExecuted)
+				]
+			];
+	
+
+	FStatusBarDrawer OutputLogDrawer(StatusBarDrawerIds::OutputLog);
+
+	OutputLogDrawer.GetDrawerContentDelegate.BindUObject(this, &UStatusBarSubsystem::OnGetOutputLog);
+	OutputLogDrawer.OnDrawerOpenedDelegate.BindUObject(this, &UStatusBarSubsystem::OnOutputLogOpened);
+	OutputLogDrawer.OnDrawerDismissedDelegate.BindUObject(this, &UStatusBarSubsystem::OnOutputLogDismised);
+	OutputLogDrawer.CustomWidget = OutputLog;
+
+	OutputLogDrawer.ButtonText = LOCTEXT("StatusBar_OutputLogButton", "Output Log");
+	OutputLogDrawer.ToolTipText = FText::Format(LOCTEXT("StatusBar_OutputLogButtonTip", "Opens a temporary output log which will dismiss when it loses focus ({0})"), FGlobalEditorCommonCommands::Get().OpenConsoleCommandBox->GetInputText());
+	OutputLogDrawer.Icon = FAppStyle::Get().GetBrush("Log.TabIcon");
+	StatusBar->RegisterDrawer(MoveTemp(OutputLogDrawer));
+
 	// Clean up stale status bars
 	for (auto It = StatusBars.CreateIterator(); It; ++It)
 	{
-		if (!It.Value().IsValid())
+		if (!It.Value().StatusBarWidget.IsValid())
 		{
 			It.RemoveCurrent();
 		}
 	}
 
-	StatusBars.Add(StatusBarName, StatusBar);
+	FStatusBarData StatusBarData;
+	StatusBarData.StatusBarWidget = StatusBar;
+	StatusBarData.ConsoleEditBox = ConsoleEditBox;
+
+	StatusBars.Add(StatusBarName, StatusBarData);
 
 	return StatusBar;
 }
@@ -393,7 +501,7 @@ void UStatusBarSubsystem::StartProgressNotification(FProgressNotificationHandle
 	// Find the active status bar to display the progress in
 	for (auto StatusBar : StatusBars)
 	{
-		if (TSharedPtr<SStatusBar> StatusBarPinned = StatusBar.Value.Pin())
+		if (TSharedPtr<SStatusBar> StatusBarPinned = StatusBar.Value.StatusBarWidget.Pin())
 		{
 			TSharedPtr<SDockTab> ParentTab = StatusBarPinned->GetParentTab();
 			if (ParentTab && ParentTab->IsForeground() && ParentTab->GetParentWindow() == ActiveWindow)
@@ -409,7 +517,7 @@ void UStatusBarSubsystem::UpdateProgressNotification(FProgressNotificationHandle
 {
 	for (auto StatusBar : StatusBars)
 	{
-		if (TSharedPtr<SStatusBar> StatusBarPinned = StatusBar.Value.Pin())
+		if (TSharedPtr<SStatusBar> StatusBarPinned = StatusBar.Value.StatusBarWidget.Pin())
 		{
 			if (StatusBarPinned->UpdateProgressNotification(Handle, TotalWorkDone, UpdatedTotalWorkToDo, UpdatedDisplayText))
 			{
@@ -424,7 +532,7 @@ void UStatusBarSubsystem::CancelProgressNotification(FProgressNotificationHandle
 {
 	for (auto StatusBar : StatusBars)
 	{
-		if (TSharedPtr<SStatusBar> StatusBarPinned = StatusBar.Value.Pin())
+		if (TSharedPtr<SStatusBar> StatusBarPinned = StatusBar.Value.StatusBarWidget.Pin())
 		{
 			if (StatusBarPinned->CancelProgressNotification(Handle))
 			{
@@ -434,12 +542,12 @@ void UStatusBarSubsystem::CancelProgressNotification(FProgressNotificationHandle
 	}
 }
 
-void UStatusBarSubsystem::OnDebugConsoleClosed()
+void UStatusBarSubsystem::OnDebugConsoleClosed(TWeakPtr<SStatusBar> OwningStatusBar)
 {
-	if (PreviousKeyboardFocusedWidget.IsValid())
+	if (TSharedPtr<SStatusBar> OwningStatusBarPinned = OwningStatusBar.Pin())
 	{
-		FSlateApplication::Get().SetKeyboardFocus(PreviousKeyboardFocusedWidget.Pin());
-		PreviousKeyboardFocusedWidget.Reset();
+		TSharedPtr<SWindow> OwningWindow = OwningStatusBarPinned->GetParentTab()->GetParentWindow();
+		ToggleDebugConsole(OwningWindow.ToSharedRef());
 	}
 }
 
@@ -457,7 +565,7 @@ void UStatusBarSubsystem::CreateContentBrowserIfNeeded()
 			{
 				for (auto StatusBar : StatusBars)
 				{
-					if (TSharedPtr<SStatusBar> StatusBarPinned = StatusBar.Value.Pin())
+					if (TSharedPtr<SStatusBar> StatusBarPinned = StatusBar.Value.StatusBarWidget.Pin())
 					{
 						if (StatusBarPinned->IsDrawerOpened(StatusBarDrawerIds::ContentBrowser))
 						{
@@ -515,7 +623,7 @@ void UStatusBarSubsystem::CreateAndShowNewUserTipIfNeeded(TSharedPtr<SWindow> Pa
 
 TSharedPtr<SStatusBar> UStatusBarSubsystem::GetStatusBar(FName StatusBarName) const
 {
-	return StatusBars.FindRef(StatusBarName).Pin();
+	return StatusBars.FindRef(StatusBarName).StatusBarWidget.Pin();
 }
 
 TSharedRef<SWidget> UStatusBarSubsystem::OnGetContentBrowser()
@@ -532,9 +640,9 @@ void UStatusBarSubsystem::OnContentBrowserOpened(TSharedRef<SStatusBar>& StatusB
 	// Dismiss any other content browser that is opened when one status bar opens it.  The content browser is a shared resource and shouldn't be in the layout twice
 	for (auto StatusBar : StatusBars)
 	{
-		if (TSharedPtr<SStatusBar> StatusBarPinned = StatusBar.Value.Pin())
+		if (TSharedPtr<SStatusBar> StatusBarPinned = StatusBar.Value.StatusBarWidget.Pin())
 		{
-			if (StatusBarWithContentBrowser != StatusBarPinned && StatusBarPinned->IsDrawerOpened(StatusBarDrawerIds::ContentBrowser))
+			if (StatusBarWithContentBrowser != StatusBarPinned || StatusBarPinned->IsAnyOtherDrawerOpened(StatusBarDrawerIds::ContentBrowser))
 			{
 				StatusBarPinned->CloseDrawerImmediately();
 			}
@@ -566,13 +674,65 @@ void UStatusBarSubsystem::HandleDeferredOpenContentBrowser(TSharedPtr<SWindow> P
 {
 	for (auto StatusBar : StatusBars)
 	{
-		if (TSharedPtr<SStatusBar> StatusBarPinned = StatusBar.Value.Pin())
+		if (TSharedPtr<SStatusBar> StatusBarPinned = StatusBar.Value.StatusBarWidget.Pin())
 		{
 			TSharedPtr<SDockTab> ParentTab = StatusBarPinned->GetParentTab();
 			if (ParentTab && ParentTab->IsForeground() && ParentTab->GetParentWindow() == ParentWindow)
 			{
 				StatusBarPinned->OpenDrawer(StatusBarDrawerIds::ContentBrowser);
-				IContentBrowserSingleton& ContentBrowserSingleton = FModuleManager::Get().GetModuleChecked<FContentBrowserModule>("ContentBrowser").Get();
+				break;
+			}
+		}
+	}
+}
+
+TSharedRef<SWidget> UStatusBarSubsystem::OnGetOutputLog()
+{
+	if (!StatusBarOutputLog)
+	{
+		StatusBarOutputLog = FOutputLogModule::Get().MakeOutputLogDrawerWidget(FSimpleDelegate::CreateUObject(this, &UStatusBarSubsystem::OnDebugConsoleDrawerClosed));
+	}
+
+	return StatusBarOutputLog.ToSharedRef();
+}
+
+void UStatusBarSubsystem::OnOutputLogOpened(TSharedRef<SStatusBar>& StatusBarWithContentBrowser)
+{
+	// Dismiss any other content browser that is opened when one status bar opens it.  The content browser is a shared resource and shouldn't be in the layout twice
+	for (auto StatusBar : StatusBars)
+	{
+		if (TSharedPtr<SStatusBar> StatusBarPinned = StatusBar.Value.StatusBarWidget.Pin())
+		{
+			if (StatusBarWithContentBrowser != StatusBarPinned || StatusBarPinned->IsAnyOtherDrawerOpened(StatusBarDrawerIds::OutputLog))
+			{
+				StatusBarPinned->CloseDrawerImmediately();
+			}
+		}
+	}
+
+	FOutputLogModule::Get().FocusOutputLogConsoleBox(StatusBarOutputLog.ToSharedRef());
+}
+
+void UStatusBarSubsystem::OnOutputLogDismised(const TSharedPtr<SWidget>& NewlyFocusedWidget)
+{
+	if (PreviousKeyboardFocusedWidget.IsValid() && !NewlyFocusedWidget.IsValid())
+	{
+		FSlateApplication::Get().SetKeyboardFocus(PreviousKeyboardFocusedWidget.Pin());
+	}
+
+	PreviousKeyboardFocusedWidget.Reset();
+
+}
+
+void UStatusBarSubsystem::OnDebugConsoleDrawerClosed()
+{
+	for (auto StatusBar : StatusBars)
+	{
+		if (TSharedPtr<SStatusBar> StatusBarPinned = StatusBar.Value.StatusBarWidget.Pin())
+		{
+			if (StatusBarPinned->IsDrawerOpened(StatusBarDrawerIds::OutputLog))
+			{
+				ToggleDebugConsole(StatusBarPinned->GetParentTab()->GetParentWindow().ToSharedRef());
 				break;
 			}
 		}
diff --git a/Engine/Source/Editor/StatusBar/Public/StatusBarSubsystem.h b/Engine/Source/Editor/StatusBar/Public/StatusBarSubsystem.h
--- a/Engine/Source/Editor/StatusBar/Public/StatusBarSubsystem.h
+++ b/Engine/Source/Editor/StatusBar/Public/StatusBarSubsystem.h
@@ -47,6 +47,12 @@ private:
 	int32 Id;
 };
 
+struct FStatusBarData
+{
+	TWeakPtr<SStatusBar> StatusBarWidget;
+	TSharedPtr<SWidget> ConsoleEditBox;
+};
+
 UCLASS()
 class STATUSBAR_API UStatusBarSubsystem : public UEditorSubsystem, public IProgressNotificationHandler
 {
@@ -65,12 +71,13 @@ public:
 	virtual void Deinitialize() override final;
 
 	/**
-	 * Focuses the debug console on the status bar for status bar residing in the passed in parent window 
+	 * Focuses the debug console or opens the output log drawer on the status bar for status bar residing in the passed in parent window 
 	 *
-	 * @param ParentWindow The parent window of the status bar 
+	 * @param ParentWindow			The parent window of the status bar 
+	 * @param bAlwaysToggleDrawer	If true, the output log drawer will be toggled without focusing the debug console first
 	 * @return true of a status bar debug console was found and focused 
 	 */
-	bool FocusDebugConsole(TSharedRef<SWindow> ParentWindow);
+	bool ToggleDebugConsole(TSharedRef<SWindow> ParentWindow, bool bAlwaysToggleDrawer=false);
 
 	/**
 	 * Opens the content browser drawer for a status bar residing in the active window 
@@ -79,6 +86,13 @@ public:
 	 */
 	bool OpenContentBrowserDrawer();
 
+	/**
+	 * Opens the output log drawer for a status bar residing in the active window
+	 *
+	 * @return true if the output log was opened or false if no status bar in the active window was found
+	 */
+	bool OpenOutputLogDrawer();
+
 	/**
 	 * Forces the drawer to dismiss. Usually it dismisses with focus. Only call this if there is some reason an open drawer would be invalid for the current state of the editor.
 	 */
@@ -124,7 +138,7 @@ private:
 	virtual void CancelProgressNotification(FProgressNotificationHandle Handle) override;
 
 	bool ToggleContentBrowser(TSharedRef<SWindow> ParentWindow);
-	void OnDebugConsoleClosed();
+	void OnDebugConsoleClosed(TWeakPtr<SStatusBar> OwningStatusBar);
 	void CreateContentBrowserIfNeeded();
 	void CreateAndShowNewUserTipIfNeeded(TSharedPtr<SWindow> ParentWindow, bool bIsNewProjectDialog);
 
@@ -133,10 +147,17 @@ private:
 	void OnContentBrowserOpened(TSharedRef<SStatusBar>& StatusBarWithContentBrowser);
 	void OnContentBrowserDismissed(const TSharedPtr<SWidget>& NewlyFocusedWidget);
 	void HandleDeferredOpenContentBrowser(TSharedPtr<SWindow> ParentWindow);
+
+	TSharedRef<SWidget> OnGetOutputLog();
+	void OnOutputLogOpened(TSharedRef<SStatusBar>& StatusBarWithContentBrowser);
+	void OnOutputLogDismised(const TSharedPtr<SWidget>& NewlyFocusedWidget);
+
+	void OnDebugConsoleDrawerClosed();
 private:
-	TMap<FName, TWeakPtr<SStatusBar>> StatusBars;
+	TMap<FName, FStatusBarData> StatusBars;
 	TWeakPtr<SWidget> PreviousKeyboardFocusedWidget;
 	/** The floating content browser that is opened via the content browser button in the status bar */
 	TSharedPtr<SWidget> StatusBarContentBrowser;
+	TSharedPtr<SWidget> StatusBarOutputLog;
 	static int32 MessageHandleCounter;
 };
diff --git a/Engine/Source/Editor/UnrealEd/Private/Toolkits/GlobalEditorCommonCommands.cpp b/Engine/Source/Editor/UnrealEd/Private/Toolkits/GlobalEditorCommonCommands.cpp
--- a/Engine/Source/Editor/UnrealEd/Private/Toolkits/GlobalEditorCommonCommands.cpp
+++ b/Engine/Source/Editor/UnrealEd/Private/Toolkits/GlobalEditorCommonCommands.cpp
@@ -39,6 +39,9 @@ void FGlobalEditorCommonCommands::RegisterCommands()
 	UI_COMMAND(FindInContentBrowser, "Browse to Asset", "Browses to the associated asset and selects it in the most recently used Content Browser (summoning one if necessary)", EUserInterfaceActionType::Button, FInputChord(EModifierKey::Control, EKeys::B));
 
 	UI_COMMAND(OpenConsoleCommandBox, "Open Console Command Box", "Opens an edit box where you can type in a console command", EUserInterfaceActionType::Button, FInputChord(EKeys::Tilde));
+	
+	UI_COMMAND(OpenOutputLogDrawer, "Open Output Log Drawer", "Opens the output log drawer from the active asset editor status bar", EUserInterfaceActionType::Button, FInputChord());
+
 	UI_COMMAND(OpenDocumentation, "Open Documentation...", "Opens documentation for this tool", EUserInterfaceActionType::Button, FInputChord(EKeys::F1));
 
 #if PLATFORM_MAC
@@ -84,6 +87,10 @@ void FGlobalEditorCommonCommands::MapActions(TSharedRef<FUICommandList>& Toolkit
 	ToolkitCommands->MapAction(
 		Get().OpenContentBrowserDrawer,
 		FExecuteAction::CreateStatic(&FGlobalEditorCommonCommands::OnOpenContentBrowserDrawer));
+
+	ToolkitCommands->MapAction(
+		Get().OpenOutputLogDrawer,
+		FExecuteAction::CreateStatic(&FGlobalEditorCommonCommands::OnOpenOutputLogDrawer));
 }
 
 void FGlobalEditorCommonCommands::OnPressedCtrlTab(TSharedPtr<FUICommandInfo> TriggeringCommand)
@@ -166,7 +173,7 @@ void FGlobalEditorCommonCommands::OnSummonedConsoleCommandBox()
 		TSharedRef<SWindow> WindowRef = ParentWindow.ToSharedRef();
 		FOutputLogModule& OutputLogModule = FModuleManager::LoadModuleChecked<FOutputLogModule>(TEXT("OutputLog"));
 
-		if (!GEditor->GetEditorSubsystem<UStatusBarSubsystem>()->FocusDebugConsole(ParentWindow.ToSharedRef()))
+		if (!GEditor->GetEditorSubsystem<UStatusBarSubsystem>()->ToggleDebugConsole(ParentWindow.ToSharedRef()))
 		{
 			// A status bar was not found, pop open a floating window instead
 			FDebugConsoleDelegates Delegates;
@@ -183,6 +190,11 @@ void FGlobalEditorCommonCommands::OnOpenContentBrowserDrawer()
 	GEditor->GetEditorSubsystem<UStatusBarSubsystem>()->OpenContentBrowserDrawer();
 }
 
+void FGlobalEditorCommonCommands::OnOpenOutputLogDrawer()
+{
+	GEditor->GetEditorSubsystem<UStatusBarSubsystem>()->OpenOutputLogDrawer();
+}
+
 //////////////////////////////////////////////////////////////////////////
 
 #undef LOCTEXT_NAMESPACE
diff --git a/Engine/Source/Editor/UnrealEd/Public/Toolkits/GlobalEditorCommonCommands.h b/Engine/Source/Editor/UnrealEd/Public/Toolkits/GlobalEditorCommonCommands.h
--- a/Engine/Source/Editor/UnrealEd/Public/Toolkits/GlobalEditorCommonCommands.h
+++ b/Engine/Source/Editor/UnrealEd/Public/Toolkits/GlobalEditorCommonCommands.h
@@ -26,6 +26,7 @@ protected:
 	static void OnSummonedAssetPicker();
 	static void OnSummonedConsoleCommandBox();
 	static void OnOpenContentBrowserDrawer();
+	static void OnOpenOutputLogDrawer();
 
 	static TSharedPtr<IMenu> OpenPopupMenu(TSharedRef<SWidget> WindowContents, const FVector2D& PopupDesiredSize);
 public:
@@ -38,6 +39,7 @@ public:
 	TSharedPtr<FUICommandInfo> SummonOpenAssetDialogAlternate;
 	TSharedPtr<FUICommandInfo> OpenDocumentation;
 	TSharedPtr<FUICommandInfo> OpenConsoleCommandBox;
+	TSharedPtr<FUICommandInfo> OpenOutputLogDrawer;
 	TSharedPtr<FUICommandInfo> OpenContentBrowserDrawer;
 };
 
