Minor adjustments to the output log drawer.
- ~ key now also works for python command boxes
- Improved tooltip for the output log drawer
- Cosmetic UI changes.

diff --git a/Engine/Plugins/Experimental/PythonScriptPlugin/Source/PythonScriptPlugin/Private/PythonScriptPlugin.cpp b/Engine/Plugins/Experimental/PythonScriptPlugin/Source/PythonScriptPlugin/Private/PythonScriptPlugin.cpp
--- a/Engine/Plugins/Experimental/PythonScriptPlugin/Source/PythonScriptPlugin/Private/PythonScriptPlugin.cpp
+++ b/Engine/Plugins/Experimental/PythonScriptPlugin/Source/PythonScriptPlugin/Private/PythonScriptPlugin.cpp
@@ -46,6 +46,7 @@
 #include "ContentBrowserDataSubsystem.h"
 #include "ContentBrowserFileDataCore.h"
 #include "ContentBrowserFileDataSource.h"
+#include "Toolkits/GlobalEditorCommonCommands.h"
 #endif	// WITH_EDITOR
 
 #if PLATFORM_WINDOWS
@@ -146,7 +147,7 @@ bool FPythonCommandExecutor::Exec(const TCHAR* Input)
 
 bool FPythonCommandExecutor::AllowHotKeyClose() const
 {
-	return false;
+	return true;
 }
 
 bool FPythonCommandExecutor::AllowMultiLine() const
@@ -154,6 +155,15 @@ bool FPythonCommandExecutor::AllowMultiLine() const
 	return true;
 }
 
+FInputChord FPythonCommandExecutor::GetHotKey() const
+{
+#if WITH_EDITOR
+	return FGlobalEditorCommonCommands::Get().OpenConsoleCommandBox->GetActiveChord(EMultipleKeyBindingIndex::Primary).Get();
+#else
+	return FInputChord();
+#endif
+}
+
 FPythonREPLCommandExecutor::FPythonREPLCommandExecutor(IPythonScriptPlugin* InPythonScriptPlugin)
 	: PythonScriptPlugin(InPythonScriptPlugin)
 {
@@ -210,7 +220,7 @@ bool FPythonREPLCommandExecutor::Exec(const TCHAR* Input)
 
 bool FPythonREPLCommandExecutor::AllowHotKeyClose() const
 {
-	return false;
+	return true;
 }
 
 bool FPythonREPLCommandExecutor::AllowMultiLine() const
@@ -218,6 +228,15 @@ bool FPythonREPLCommandExecutor::AllowMultiLine() const
 	return true;
 }
 
+FInputChord FPythonREPLCommandExecutor::GetHotKey() const
+{
+#if WITH_EDITOR
+	return FGlobalEditorCommonCommands::Get().OpenConsoleCommandBox->GetActiveChord(EMultipleKeyBindingIndex::Primary).Get();
+#else
+	return FInputChord();
+#endif
+}
+
 #if WITH_EDITOR
 class FPythonCommandMenuImpl : public IPythonCommandMenu
 {
diff --git a/Engine/Plugins/Experimental/PythonScriptPlugin/Source/PythonScriptPlugin/Private/PythonScriptPlugin.h b/Engine/Plugins/Experimental/PythonScriptPlugin/Source/PythonScriptPlugin/Private/PythonScriptPlugin.h
--- a/Engine/Plugins/Experimental/PythonScriptPlugin/Source/PythonScriptPlugin/Private/PythonScriptPlugin.h
+++ b/Engine/Plugins/Experimental/PythonScriptPlugin/Source/PythonScriptPlugin/Private/PythonScriptPlugin.h
@@ -42,10 +42,8 @@ public:
 	virtual bool Exec(const TCHAR* Input) override;
 	virtual bool AllowHotKeyClose() const override;
 	virtual bool AllowMultiLine() const override;
-	virtual FInputChord GetHotKey() const override
-	{
-		return FInputChord();
-	}
+	virtual FInputChord GetHotKey() const override;
+
 private:
 	IPythonScriptPlugin* PythonScriptPlugin;
 };
@@ -68,10 +66,8 @@ public:
 	virtual bool Exec(const TCHAR* Input) override;
 	virtual bool AllowHotKeyClose() const override;
 	virtual bool AllowMultiLine() const override;
-	virtual FInputChord GetHotKey() const override
-	{
-		return FInputChord();
-	}
+	virtual FInputChord GetHotKey() const override;
+
 private:
 	IPythonScriptPlugin* PythonScriptPlugin;
 };
diff --git a/Engine/Source/Editor/StatusBar/Private/SStatusBar.cpp b/Engine/Source/Editor/StatusBar/Private/SStatusBar.cpp
--- a/Engine/Source/Editor/StatusBar/Private/SStatusBar.cpp
+++ b/Engine/Source/Editor/StatusBar/Private/SStatusBar.cpp
@@ -842,41 +842,41 @@ TSharedRef<SWidget> SStatusBar::MakeStatusBarDrawerButton(const FStatusBarDrawer
 {
 	const FName DrawerId = Drawer.UniqueId;
 
+	const FSlateBrush* StatusBarBackground = FAppStyle::Get().GetBrush("Brushes.Panel");
+
 	TSharedRef<SWidget> DrawerButton = 
-		SNew(SButton)
-		.IsFocusable(false)
-		.ButtonStyle(&FAppStyle::Get().GetWidgetStyle<FButtonStyle>("StatusBar.StatusBarButton"))
-		.OnClicked(this, &SStatusBar::OnDrawerButtonClicked, DrawerId)
-		.ToolTipText(Drawer.ToolTipText)
+
+		SNew(SBorder)
+		.Padding(FMargin(2.0f, 0.0f))
+		.BorderImage(StatusBarBackground)
+		.Visibility(EVisibility::SelfHitTestInvisible)
+		.VAlign(VAlign_Center)
 		[
-			SNew(SHorizontalBox)
-			+ SHorizontalBox::Slot()
-			.HAlign(HAlign_Left)
-			.VAlign(VAlign_Center)
-			.Padding(2.0f)
-			.AutoWidth()
-			[
-				SNew(SImage)
-				.ColorAndOpacity(FSlateColor::UseForeground())
-				.Image_Lambda([this, DrawerId]() { return IsDrawerOpened(DrawerId) ? DownArrow : UpArrow; })
-			]
-			+ SHorizontalBox::Slot()
-			.Padding(2.0f)
-			.HAlign(HAlign_Left)
-			.VAlign(VAlign_Center)
-			.AutoWidth()
-			[
-				SNew(SImage)
-				.ColorAndOpacity(FSlateColor::UseForeground())
-				.Image(Drawer.Icon)
-			]
-			+ SHorizontalBox::Slot()
-			.VAlign(VAlign_Center)
-			.Padding(2.0f)
+			SNew(SButton)
+			.IsFocusable(false)
+			.ButtonStyle(&FAppStyle::Get().GetWidgetStyle<FButtonStyle>("StatusBar.StatusBarButton"))
+			.OnClicked(this, &SStatusBar::OnDrawerButtonClicked, DrawerId)
+			.ToolTipText(Drawer.ToolTipText)
 			[
-				SNew(STextBlock)
-				.TextStyle(&FAppStyle::Get().GetWidgetStyle<FTextBlockStyle>("NormalText"))
-				.Text(Drawer.ButtonText)
+				SNew(SHorizontalBox)
+				+ SHorizontalBox::Slot()
+				.Padding(2.0f)
+				.HAlign(HAlign_Left)
+				.VAlign(VAlign_Center)
+				.AutoWidth()
+				[
+					SNew(SImage)
+					.ColorAndOpacity(FSlateColor::UseForeground())
+					.Image(Drawer.Icon)
+				]
+				+ SHorizontalBox::Slot()
+				.VAlign(VAlign_Center)
+				.Padding(2.0f)
+				[
+					SNew(STextBlock)
+					.TextStyle(&FAppStyle::Get().GetWidgetStyle<FTextBlockStyle>("NormalText"))
+					.Text(Drawer.ButtonText)
+				]
 			]
 		];
 
@@ -886,14 +886,21 @@ TSharedRef<SWidget> SStatusBar::MakeStatusBarDrawerButton(const FStatusBarDrawer
 		return
 			SNew(SHorizontalBox)
 			+ SHorizontalBox::Slot()
-			.Padding(0.0f, 0.0f, 4.0f, 0.0f)
+			.Padding(0.0f, 0.0f, 2.0f, 0.0f)
 			.AutoWidth()
 			[
 				DrawerButton
 			]
 			+ SHorizontalBox::Slot()
 			[
-				Drawer.CustomWidget.ToSharedRef()
+				SNew(SBorder)
+				.Padding(FMargin(2.0f, 0.0f))
+				.BorderImage(StatusBarBackground)
+				.Visibility(EVisibility::SelfHitTestInvisible)
+				.VAlign(VAlign_Center)
+				[
+					Drawer.CustomWidget.ToSharedRef()
+				]
 			];
 	
 	}
@@ -1185,19 +1192,11 @@ void SStatusBar::RegisterDrawer(FStatusBarDrawer&& Drawer)
 
 	if (RegisteredDrawers.Num() > NumDrawers)
 	{
-		const FSlateBrush* StatusBarBackground = FAppStyle::Get().GetBrush("Brushes.Panel");
-
 		DrawerBox->AddSlot()
 		.Padding(1.0f, 0.0f)
 		.AutoWidth()
 		[
-			SNew(SBorder)
-			.Padding(FMargin(2.0f, 0.0f))
-			.BorderImage(StatusBarBackground)
-			.VAlign(VAlign_Center)
-			[
-				MakeStatusBarDrawerButton(Drawer)
-			]
+			MakeStatusBarDrawerButton(Drawer)
 		];
 	}
 }
diff --git a/Engine/Source/Editor/StatusBar/Private/StatusBarSubsystem.cpp b/Engine/Source/Editor/StatusBar/Private/StatusBarSubsystem.cpp
--- a/Engine/Source/Editor/StatusBar/Private/StatusBarSubsystem.cpp
+++ b/Engine/Source/Editor/StatusBar/Private/StatusBarSubsystem.cpp
@@ -431,7 +431,7 @@ TSharedRef<SWidget> UStatusBarSubsystem::MakeStatusBarWidget(FName StatusBarName
 	OutputLogDrawer.CustomWidget = OutputLog;
 
 	OutputLogDrawer.ButtonText = LOCTEXT("StatusBar_OutputLogButton", "Output Log");
-	OutputLogDrawer.ToolTipText = FText::Format(LOCTEXT("StatusBar_OutputLogButtonTip", "Opens a temporary output log which will dismiss when it loses focus ({0})"), FGlobalEditorCommonCommands::Get().OpenConsoleCommandBox->GetInputText());
+	OutputLogDrawer.ToolTipText = FText::Format(LOCTEXT("StatusBar_OutputLogButtonTip", "Opens the output log drawer. ({0}) cycles between focusing the console command box, opening the output log drawer, and closing it.\nThe output log drawer may also be toggled directly with ({1})"), FGlobalEditorCommonCommands::Get().OpenConsoleCommandBox->GetInputText(), FGlobalEditorCommonCommands::Get().OpenOutputLogDrawer->GetInputText());
 	OutputLogDrawer.Icon = FAppStyle::Get().GetBrush("Log.TabIcon");
 	StatusBar->RegisterDrawer(MoveTemp(OutputLogDrawer));
 
diff --git a/Engine/Source/Editor/UnrealEd/Private/Toolkits/GlobalEditorCommonCommands.cpp b/Engine/Source/Editor/UnrealEd/Private/Toolkits/GlobalEditorCommonCommands.cpp
--- a/Engine/Source/Editor/UnrealEd/Private/Toolkits/GlobalEditorCommonCommands.cpp
+++ b/Engine/Source/Editor/UnrealEd/Private/Toolkits/GlobalEditorCommonCommands.cpp
@@ -40,7 +40,7 @@ void FGlobalEditorCommonCommands::RegisterCommands()
 
 	UI_COMMAND(OpenConsoleCommandBox, "Open Console Command Box", "Opens an edit box where you can type in a console command", EUserInterfaceActionType::Button, FInputChord(EKeys::Tilde));
 	
-	UI_COMMAND(OpenOutputLogDrawer, "Open Output Log Drawer", "Opens the output log drawer from the active asset editor status bar", EUserInterfaceActionType::Button, FInputChord());
+	UI_COMMAND(OpenOutputLogDrawer, "Open Output Log Drawer", "Opens the output log drawer from the active asset editor status bar", EUserInterfaceActionType::Button, FInputChord(EModifierKey::Alt, EKeys::Tilde));
 
 	UI_COMMAND(OpenDocumentation, "Open Documentation...", "Opens documentation for this tool", EUserInterfaceActionType::Button, FInputChord(EKeys::F1));
 
