Fix console edit box not appearing in immersive mode

diff --git a/Engine/Source/Developer/OutputLog/Private/OutputLogModule.cpp b/Engine/Source/Developer/OutputLog/Private/OutputLogModule.cpp
--- a/Engine/Source/Developer/OutputLog/Private/OutputLogModule.cpp
+++ b/Engine/Source/Developer/OutputLog/Private/OutputLogModule.cpp
@@ -235,6 +235,12 @@ void FOutputLogModule::CloseDebugConsole()
 		{
 			WindowForExistingConsole->RemoveOverlaySlot(PinnedDebugConsole.ToSharedRef());
 			DebugConsole.Reset();
+
+			if (TSharedPtr<SWidget> PreviousKeyboardFocusedWidgetPinned = PreviousKeyboardFocusedWidget.Pin())
+			{
+				FSlateApplication::Get().SetKeyboardFocus(PreviousKeyboardFocusedWidgetPinned);
+				PreviousKeyboardFocusedWidget.Reset();
+			}
 		}
 	}
 }
diff --git a/Engine/Source/Editor/StatusBar/Private/StatusBarSubsystem.cpp b/Engine/Source/Editor/StatusBar/Private/StatusBarSubsystem.cpp
--- a/Engine/Source/Editor/StatusBar/Private/StatusBarSubsystem.cpp
+++ b/Engine/Source/Editor/StatusBar/Private/StatusBarSubsystem.cpp
@@ -242,47 +242,47 @@ bool UStatusBarSubsystem::ToggleDebugConsole(TSharedRef<SWindow> ParentWindow, b
 			TSharedPtr<SDockTab> ParentTab = StatusBarPinned->GetParentTab();
 			if (ParentTab && ParentTab->IsForeground() && ParentTab->GetParentWindow() == ParentWindow)
 			{
-				if (!bAlwaysToggleDrawer && MainOutputLogTab && MainOutputLogTab->GetTabManager()->GetOwnerTab() == ParentTab)
+				FWidgetPath ConsoleEditBoxPath;
+				if(FSlateApplication::Get().GeneratePathToWidgetUnchecked(SBData.ConsoleEditBox.ToSharedRef(), ConsoleEditBoxPath))
 				{
-					OutputLogModule.FocusOutputLogConsoleBox(OutputLogModule.GetOutputLog().ToSharedRef());
-				}
-				else
-				{
-					// This toggles between 3 states: 
-					// If the drawer is opened, close it
-					// if the console edit box is focused, open the drawer
-					// if something else is focused, focus the console edit box
-					if (StatusBarPinned->IsDrawerOpened(StatusBarDrawerIds::OutputLog))
+					if (!!bAlwaysToggleDrawer && MainOutputLogTab && MainOutputLogTab->GetTabManager()->GetOwnerTab() == ParentTab)
 					{
-						StatusBarPinned->DismissDrawer(nullptr);
+						OutputLogModule.FocusOutputLogConsoleBox(OutputLogModule.GetOutputLog().ToSharedRef());
 					}
-					else if (SBData.ConsoleEditBox->HasKeyboardFocus() || bAlwaysToggleDrawer)
+					else
 					{
-						if (bCycleToOutputLogDrawer)
+						// This toggles between 3 states: 
+						// If the drawer is opened, close it
+						// if the console edit box is focused, open the drawer
+						// if something else is focused, focus the console edit box
+						if (StatusBarPinned->IsDrawerOpened(StatusBarDrawerIds::OutputLog))
 						{
-							StatusBarPinned->OpenDrawer(StatusBarDrawerIds::OutputLog);
+							StatusBarPinned->DismissDrawer(nullptr);
+						}
+						else if (SBData.ConsoleEditBox->HasKeyboardFocus() || bAlwaysToggleDrawer)
+						{
+							if (bCycleToOutputLogDrawer)
+							{
+								StatusBarPinned->OpenDrawer(StatusBarDrawerIds::OutputLog);
+							}
+							else
+							{
+								// Restore previous focus
+								FSlateApplication::Get().SetKeyboardFocus(PreviousKeyboardFocusedWidget.Pin(), EFocusCause::SetDirectly);
+								PreviousKeyboardFocusedWidget.Reset();
+							}
 						}
 						else
 						{
-							// Restore previous focus
-							FSlateApplication::Get().SetKeyboardFocus(PreviousKeyboardFocusedWidget.Pin(), EFocusCause::SetDirectly);
-
-							PreviousKeyboardFocusedWidget.Reset();
+							// Cache off the previously focused widget so we can restore focus if the user hits the focus key again
+							PreviousKeyboardFocusedWidget = FSlateApplication::Get().GetKeyboardFocusedWidget();
+							FSlateApplication::Get().SetKeyboardFocus(ConsoleEditBoxPath, EFocusCause::SetDirectly);
 						}
 					}
-					else
-					{
-
-						// Cache off the previously focused widget so we can restore focus if the user hits the focus key again
-						PreviousKeyboardFocusedWidget = FSlateApplication::Get().GetKeyboardFocusedWidget();
 
-						TSharedPtr<SWidget> WidgetToFocus = SBData.ConsoleEditBox;
-						FSlateApplication::Get().SetKeyboardFocus(WidgetToFocus, EFocusCause::SetDirectly);
-					}
+					bToggledSuccessfully = true;
+					break;
 				}
-
-				bToggledSuccessfully = true;
-				break;
 			}
 		}
 	}
