Fix source control submenu not updating each time menu is opened

diff --git a/Engine/Source/Editor/StatusBar/Private/SourceControlMenuHelpers.cpp b/Engine/Source/Editor/StatusBar/Private/SourceControlMenuHelpers.cpp
--- a/Engine/Source/Editor/StatusBar/Private/SourceControlMenuHelpers.cpp
+++ b/Engine/Source/Editor/StatusBar/Private/SourceControlMenuHelpers.cpp
@@ -272,10 +272,8 @@ TSharedRef<SWidget> FSourceControlMenuHelpers::MakeSourceControlStatusWidget()
 				.Text_Static(&FSourceControlMenuHelpers::GetSourceControlStatusText)
 			]
 		]
-		.MenuContent()
-		[
-			GenerateSourceControlMenuContent()
-		];
+		.OnGetMenuContent(FOnGetContent::CreateStatic(&FSourceControlMenuHelpers::GenerateSourceControlMenuContent))
+		;
 }
 
 #undef LOCTEXT_NAMESPACE
\ No newline at end of file
