Logging and safeguard for docking content browser drawer crash

diff --git a/Engine/Source/Editor/ContentBrowser/Private/ContentBrowserSingleton.cpp b/Engine/Source/Editor/ContentBrowser/Private/ContentBrowserSingleton.cpp
--- a/Engine/Source/Editor/ContentBrowser/Private/ContentBrowserSingleton.cpp
+++ b/Engine/Source/Editor/ContentBrowser/Private/ContentBrowserSingleton.cpp
@@ -676,6 +676,7 @@ void FContentBrowserSingleton::DockContentBrowserDrawer()
 	
 	// Dont summon a content browser if a content browser already exists in the tab manager
 	bool bExistingTab = false;
+	if(ForTabManager)
 	{
 		for (int32 BrowserIdx = 0; BrowserIdx < UE_ARRAY_COUNT(ContentBrowserTabIDs); BrowserIdx++)
 		{
diff --git a/Engine/Source/Editor/StatusBar/Private/StatusBarSubsystem.cpp b/Engine/Source/Editor/StatusBar/Private/StatusBarSubsystem.cpp
--- a/Engine/Source/Editor/StatusBar/Private/StatusBarSubsystem.cpp
+++ b/Engine/Source/Editor/StatusBar/Private/StatusBarSubsystem.cpp
@@ -27,6 +27,8 @@
 
 #define LOCTEXT_NAMESPACE "StatusBar"
 
+DEFINE_LOG_CATEGORY_STATIC(LogStatusBar, Log, All);
+
 int32 UStatusBarSubsystem::MessageHandleCounter = 0;
 
 class SNewUserTipNotification : public SCompoundWidget
@@ -563,18 +565,24 @@ void UStatusBarSubsystem::CreateContentBrowserIfNeeded()
 		TFunction<TSharedPtr<SDockTab>()> GetTab(
 			[this]() -> TSharedPtr<SDockTab>
 			{
+				UE_LOG(LogStatusBar, Log, TEXT("Looking status bar with open content browser drawer..."))
 				for (auto StatusBar : StatusBars)
 				{
 					if (TSharedPtr<SStatusBar> StatusBarPinned = StatusBar.Value.StatusBarWidget.Pin())
 					{
 						if (StatusBarPinned->IsDrawerOpened(StatusBarDrawerIds::ContentBrowser))
 						{
+							UE_LOG(LogStatusBar, Log, TEXT("Using status bar: %s to dock content browser"), *StatusBar.Key.ToString());
 							return StatusBarPinned->GetParentTab();
 						}
+						else
+						{
+							UE_LOG(LogStatusBar, Log, TEXT("StatusBar: %s was content browser was not opened"), *StatusBar.Key.ToString());
+						}
 					}
 				}
 
-				checkf(false, TEXT("If we get here somehow a content browser drawer is opened but no status bar claims it"));
+				ensureMsgf(false, TEXT("If we get here somehow a content browser drawer is opened but no status bar claims it"));
 				return TSharedPtr<SDockTab>();
 			}
 		);
