fix ordering of notification tests

diff --git a/Engine/Source/Editor/StatusBar/Private/NotificationTesting.cpp b/Engine/Source/Editor/StatusBar/Private/NotificationTesting.cpp
--- a/Engine/Source/Editor/StatusBar/Private/NotificationTesting.cpp
+++ b/Engine/Source/Editor/StatusBar/Private/NotificationTesting.cpp
@@ -113,7 +113,7 @@ static void TestNotifications()
 
 	FTicker::GetCoreTicker().AddTicker(TEXT("TestNotifications"), 4.0f, [](float DeltaTime)
 		{
-			FNotificationInfo NotificationInfo(FText::FromString(TEXT("Test Notification 5")));
+			FNotificationInfo NotificationInfo(FText::FromString(TEXT("Test Notification 4")));
 			NotificationInfo.FadeInDuration = 2.0f;
 			NotificationInfo.FadeOutDuration = 2.0f;
 			NotificationInfo.ExpireDuration = Timeout;
@@ -128,7 +128,7 @@ static void TestNotifications()
 
 	FTicker::GetCoreTicker().AddTicker(TEXT("TestNotifications"), 5.0f, [](float DeltaTime)
 		{
-			FNotificationInfo NotificationInfo(FText::FromString(TEXT("Test Notification 6")));
+			FNotificationInfo NotificationInfo(FText::FromString(TEXT("Test Notification 5")));
 			NotificationInfo.FadeInDuration = 2.0f;
 			NotificationInfo.FadeOutDuration = 2.0f;
 			NotificationInfo.ExpireDuration = Timeout;
@@ -146,7 +146,7 @@ static void TestNotifications()
 
 	FTicker::GetCoreTicker().AddTicker(TEXT("TestNotifications"), 6.0f, [](float DeltaTime)
 		{
-			FNotificationInfo NotificationInfo(FText::FromString(TEXT("Test Notification 7")));
+			FNotificationInfo NotificationInfo(FText::FromString(TEXT("Test Notification 6")));
 			NotificationInfo.FadeInDuration = 2.0f;
 			NotificationInfo.FadeOutDuration = 2.0f;
 			NotificationInfo.ExpireDuration = Timeout;
@@ -163,7 +163,7 @@ static void TestNotifications()
 
 	FTicker::GetCoreTicker().AddTicker(TEXT("TestNotifications"), 7.0f, [](float DeltaTime)
 		{
-			FNotificationInfo NotificationInfo(FText::FromString(TEXT("Test Notification 8")));
+			FNotificationInfo NotificationInfo(FText::FromString(TEXT("Test Notification 7")));
 			NotificationInfo.FadeInDuration = 2.0f;
 			NotificationInfo.FadeOutDuration = 2.0f;
 			NotificationInfo.ExpireDuration = Timeout;
