Update delay timer for status bar progress notifications to avoid popups for quick completing async tasks

diff --git a/Engine/Source/Editor/StatusBar/Private/SStatusBar.cpp b/Engine/Source/Editor/StatusBar/Private/SStatusBar.cpp
--- a/Engine/Source/Editor/StatusBar/Private/SStatusBar.cpp
+++ b/Engine/Source/Editor/StatusBar/Private/SStatusBar.cpp
@@ -38,7 +38,7 @@ namespace StatusBarNotificationConstants
 	const float NotificationFadeDuration = .15f;
 
 	// Delay before a progress notification becomes visible. This is to avoid the status bar to animate and flicker from short lived notifications. 
-	const double NotificationDelay = 0.5;
+	const double NotificationDelay = 3.0;
 }
 
 class SDrawerOverlay : public SCompoundWidget
