Adds Submit Content menu entry.
Moves Source Control related actions to Tools menu.

#changelist validated

diff --git a/Engine/Source/Editor/MainFrame/Private/Frame/MainFrameActions.cpp b/Engine/Source/Editor/MainFrame/Private/Frame/MainFrameActions.cpp
--- a/Engine/Source/Editor/MainFrame/Private/Frame/MainFrameActions.cpp
+++ b/Engine/Source/Editor/MainFrame/Private/Frame/MainFrameActions.cpp
@@ -51,6 +51,7 @@
 #include "AnalyticsEventAttribute.h"
 #include "Kismet2/DebuggerCommands.h"
 #include "GameMapsSettings.h"
+#include "SourceControlWindows.h"
 
 #define LOCTEXT_NAMESPACE "MainFrameActions"
 
@@ -99,9 +100,12 @@ void FMainFrameCommands::RegisterCommands()
 	UI_COMMAND( ChooseFilesToSave, "Choose Files to Save...", "Opens a dialog with save options for content and levels", EUserInterfaceActionType::Button, FInputChord() );
 	ActionList->MapAction( ChooseFilesToSave, FExecuteAction::CreateStatic( &FMainFrameActionCallbacks::ChoosePackagesToSave ), FCanExecuteAction::CreateStatic( &FMainFrameActionCallbacks::CanSaveWorld ) );
 
-	UI_COMMAND( ViewChangelists, "View Changelists...", "Opens a dialog with check in options for content and levels", EUserInterfaceActionType::Button, FInputChord() );
+	UI_COMMAND( ViewChangelists, "View Changelists", "Opens a dialog displaying current changelists.", EUserInterfaceActionType::Button, FInputChord() );
 	ActionList->MapAction(ViewChangelists, FExecuteAction::CreateStatic( &FMainFrameActionCallbacks::ViewChangelists ), FCanExecuteAction::CreateStatic( &FMainFrameActionCallbacks::CanViewChangelists ) );
 
+	UI_COMMAND( SubmitContent, "Submit Content", "Opens a dialog with check in options for content and levels.", EUserInterfaceActionType::Button, FInputChord() );
+	ActionList->MapAction( SubmitContent, FExecuteAction::CreateLambda([]() { FSourceControlWindows::ChoosePackagesToCheckIn(); }), FCanExecuteAction::CreateStatic(&FSourceControlWindows::CanChoosePackagesToCheckIn ) );
+
 	UI_COMMAND( ConnectToSourceControl, "Connect to Source Control...", "Connect to source control to allow source control operations to be performed on content and levels.", EUserInterfaceActionType::Button, FInputChord() );
 	ActionList->MapAction( ConnectToSourceControl, FExecuteAction::CreateStatic( &FMainFrameActionCallbacks::ConnectToSourceControl ), DefaultExecuteAction );
 
diff --git a/Engine/Source/Editor/MainFrame/Private/Frame/MainFrameActions.h b/Engine/Source/Editor/MainFrame/Private/Frame/MainFrameActions.h
--- a/Engine/Source/Editor/MainFrame/Private/Frame/MainFrameActions.h
+++ b/Engine/Source/Editor/MainFrame/Private/Frame/MainFrameActions.h
@@ -27,6 +27,7 @@ public:
 	TSharedPtr< FUICommandInfo > Exit;
 	TSharedPtr< FUICommandInfo > ChooseFilesToSave;
 	TSharedPtr< FUICommandInfo > ViewChangelists;
+	TSharedPtr< FUICommandInfo > SubmitContent;
 	TSharedPtr< FUICommandInfo > ConnectToSourceControl;
 	TSharedPtr< FUICommandInfo > ChangeSourceControlSettings;
 	TSharedPtr< FUICommandInfo > NewProject;
diff --git a/Engine/Source/Editor/MainFrame/Private/Menus/MainMenu.cpp b/Engine/Source/Editor/MainFrame/Private/Menus/MainMenu.cpp
--- a/Engine/Source/Editor/MainFrame/Private/Menus/MainMenu.cpp
+++ b/Engine/Source/Editor/MainFrame/Private/Menus/MainMenu.cpp
@@ -156,30 +156,6 @@ void FMainMenu::RegisterEditMenu()
 				FGlobalTabmanager::Get()->PopulateTabSpawnerMenu(InBuilder, "PluginsEditor");
 			}
 		}));
-
-		Section.AddDynamicEntry("ConnectToSourceControl", FNewToolMenuSectionDelegate::CreateLambda([](FToolMenuSection& InSection)
-		{
-			ISourceControlModule& SourceControlModule = ISourceControlModule::Get();
-			
-			if (ISourceControlModule::Get().IsEnabled() && ISourceControlModule::Get().GetProvider().IsAvailable())
-			{
-				InSection.AddMenuEntry(
-					FMainFrameCommands::Get().ChangeSourceControlSettings,
-					TAttribute<FText>(),
-					TAttribute<FText>(),
-					FSlateIcon(FEditorStyle::GetStyleSetName(), "SourceControl.Actions.ChangeSettings")
-				);
-			}
-			else
-			{
-				InSection.AddMenuEntry(
-					FMainFrameCommands::Get().ConnectToSourceControl,
-					TAttribute<FText>(),
-					TAttribute<FText>(),
-					FSlateIcon(FEditorStyle::GetStyleSetName(), "SourceControl.Actions.Connect")
-				);
-			}
-		}));
 	}
 }
 
@@ -553,6 +529,45 @@ void FMainMenu::RegisterToolsMenu()
 		}
 	}
 
+	FToolMenuSection& SourceControlSection = Menu->AddSection("Source Control", LOCTEXT("SourceControlHeading", "Source Control"));
+
+	SourceControlSection.AddMenuEntry(
+		FMainFrameCommands::Get().ViewChangelists,
+		TAttribute<FText>(),
+		TAttribute<FText>(),
+		FSlateIcon(FEditorStyle::GetStyleSetName(), "SourceControl.ChangelistsTab")
+	);
+
+	SourceControlSection.AddMenuEntry(
+		FMainFrameCommands::Get().SubmitContent,
+		TAttribute<FText>(),
+		TAttribute<FText>(),
+		FSlateIcon(FEditorStyle::GetStyleSetName(), "SourceControl.Actions.Submit")
+	);
+
+	SourceControlSection.AddDynamicEntry("ConnectToSourceControl", FNewToolMenuSectionDelegate::CreateLambda([](FToolMenuSection& InSection)
+	{
+		ISourceControlModule& SourceControlModule = ISourceControlModule::Get();
+
+		if (ISourceControlModule::Get().IsEnabled() && ISourceControlModule::Get().GetProvider().IsAvailable())
+		{
+			InSection.AddMenuEntry(
+				FMainFrameCommands::Get().ChangeSourceControlSettings,
+				TAttribute<FText>(),
+				TAttribute<FText>(),
+				FSlateIcon(FEditorStyle::GetStyleSetName(), "SourceControl.Actions.ChangeSettings")
+			);
+		}
+		else
+		{
+			InSection.AddMenuEntry(
+				FMainFrameCommands::Get().ConnectToSourceControl,
+				TAttribute<FText>(),
+				TAttribute<FText>(),
+				FSlateIcon(FEditorStyle::GetStyleSetName(), "SourceControl.Actions.Connect")
+			);
+		}
+	}));
 }
 
 void FMainMenu::RegisterExitMenuItems()
diff --git a/Engine/Source/Editor/SourceControlWindows/Private/SourceControlWindows.cpp b/Engine/Source/Editor/SourceControlWindows/Private/SourceControlWindows.cpp
--- a/Engine/Source/Editor/SourceControlWindows/Private/SourceControlWindows.cpp
+++ b/Engine/Source/Editor/SourceControlWindows/Private/SourceControlWindows.cpp
@@ -129,7 +129,11 @@ bool FSourceControlWindows::ChoosePackagesToCheckIn(const FSourceControlWindowsO
 
 bool FSourceControlWindows::CanChoosePackagesToCheckIn()
 {
-	return !ChoosePackagesToCheckInNotification.IsValid();
+	ISourceControlModule& SourceControlModule = ISourceControlModule::Get();
+	
+	return ISourceControlModule::Get().IsEnabled() &&
+		ISourceControlModule::Get().GetProvider().IsAvailable() &&
+		!ChoosePackagesToCheckInNotification.IsValid();
 }
 
 
diff --git a/Engine/Source/Editor/SourceControlWindows/Private/SourceControlWindowsModule.cpp b/Engine/Source/Editor/SourceControlWindows/Private/SourceControlWindowsModule.cpp
--- a/Engine/Source/Editor/SourceControlWindows/Private/SourceControlWindowsModule.cpp
+++ b/Engine/Source/Editor/SourceControlWindows/Private/SourceControlWindowsModule.cpp
@@ -61,10 +61,10 @@ void FSourceControlWindowsModule::StartupModule()
 	// Create a Source Control group under the Tools category
 	const FSlateIcon SourceControlIcon(FEditorStyle::GetStyleSetName(), "SourceControl.ChangelistsTab");
 
-	// Register the changlist tab spawner
-	FGlobalTabmanager::Get()->RegisterNomadTabSpawner(SourceControlChangelistsTabName, FOnSpawnTab::CreateRaw(this, &FSourceControlWindowsModule::CreateChangelistsTab))
-		.SetDisplayName(LOCTEXT("ChangelistsTabTitle", "Changelists"))
-		.SetGroup(WorkspaceMenu::GetMenuStructure().GetToolsCategory())
+	// Register the changelist tab spawner
+	FGlobalTabmanager::Get()->RegisterTabSpawner(SourceControlChangelistsTabName, FOnSpawnTab::CreateRaw(this, &FSourceControlWindowsModule::CreateChangelistsTab))
+		.SetDisplayName(LOCTEXT("ChangelistsTabTitle", "View Changelists"))
+		.SetTooltipText(LOCTEXT("ChangelistsTabTooltip", "Opens a dialog displaying current changelists."))
 		.SetIcon(SourceControlIcon);
 
 #if WITH_RELOAD
@@ -80,7 +80,7 @@ void FSourceControlWindowsModule::ShutdownModule()
 {
 	if (FSlateApplication::IsInitialized())
 	{
-		FGlobalTabmanager::Get()->UnregisterNomadTabSpawner(SourceControlChangelistsTabName);
+		FGlobalTabmanager::Get()->UnregisterTabSpawner(SourceControlChangelistsTabName);
 
 		if (ChangelistsTab.IsValid())
 		{
diff --git a/Engine/Source/Editor/StatusBar/Private/SourceControlMenuHelpers.cpp b/Engine/Source/Editor/StatusBar/Private/SourceControlMenuHelpers.cpp
--- a/Engine/Source/Editor/StatusBar/Private/SourceControlMenuHelpers.cpp
+++ b/Engine/Source/Editor/StatusBar/Private/SourceControlMenuHelpers.cpp
@@ -14,6 +14,7 @@
 #include "EditorStyleSet.h"
 #include "Widgets/Input/SComboButton.h"
 #include "Widgets/Images/SImage.h"
+#include "LevelEditorActions.h"
 
 #define LOCTEXT_NAMESPACE "SourceControlCommands"
 
@@ -34,10 +35,11 @@ FSourceControlCommands::FSourceControlCommands()
  */
 void FSourceControlCommands::RegisterCommands()
 {
-	UI_COMMAND(ConnectToSourceControl, "Connect to Source Control...", "Opens a dialog to connect to source control.", EUserInterfaceActionType::Button, FInputChord());
+	UI_COMMAND(ConnectToSourceControl, "Connect to Source Control...", "Connect to source control to allow source control operations to be performed on content and levels.", EUserInterfaceActionType::Button, FInputChord());
 	UI_COMMAND(ChangeSourceControlSettings, "Change Source Control Settings...", "Opens a dialog to change source control settings.", EUserInterfaceActionType::Button, FInputChord());
-	UI_COMMAND(ViewChangelists, "View Changelists...", "Open the Source Control Changelist tab.", EUserInterfaceActionType::Button, FInputChord());
-	UI_COMMAND(CheckOutModifiedFiles, "Check Out Modified Files...", "Opens a dialog to check out any assets which have been modified.", EUserInterfaceActionType::Button, FInputChord());
+	UI_COMMAND(ViewChangelists, "View Changelists", "Opens a dialog displaying current changelists.", EUserInterfaceActionType::Button, FInputChord());
+	UI_COMMAND(SubmitContent, "Submit Content", "Opens a dialog with check in options for content and levels.", EUserInterfaceActionType::Button, FInputChord());
+	UI_COMMAND(CheckOutModifiedFiles, "Check Out Modified Files", "Opens a dialog to check out any assets which have been modified.", EUserInterfaceActionType::Button, FInputChord());
 
 	ActionList->MapAction(
 		ConnectToSourceControl,
@@ -55,6 +57,12 @@ void FSourceControlCommands::RegisterCommands()
 		FCanExecuteAction::CreateStatic(&FSourceControlCommands::ViewChangelists_CanExecute)
 	);
 
+	ActionList->MapAction(
+		SubmitContent,
+		FExecuteAction::CreateLambda([]() { FSourceControlWindows::ChoosePackagesToCheckIn(); }),
+		FCanExecuteAction::CreateStatic(&FSourceControlWindows::CanChoosePackagesToCheckIn)
+	);
+
 	ActionList->MapAction(
 		CheckOutModifiedFiles,
 		FExecuteAction::CreateStatic(&FSourceControlCommands::CheckOutModifiedFiles_Clicked),
@@ -131,31 +139,6 @@ TSharedRef<SWidget> FSourceControlMenuHelpers::GenerateSourceControlMenuContent(
 
 	FToolMenuSection& Section = SourceControlMenu->AddSection("SourceControlActions", LOCTEXT("SourceControlMenuHeadingActions", "Actions"));
 
-	Section.AddDynamicEntry("ConnectToSourceControl", FNewToolMenuSectionDelegate::CreateLambda([](FToolMenuSection& InSection)
-		{
-			ISourceControlModule& SourceControlModule = ISourceControlModule::Get();
-			if (ISourceControlModule::Get().IsEnabled() && ISourceControlModule::Get().GetProvider().IsAvailable())
-			{
-				InSection.AddMenuEntry(
-					FSourceControlCommands::Get().ChangeSourceControlSettings,
-					TAttribute<FText>(),
-					TAttribute<FText>(),
-					FSlateIcon(FEditorStyle::GetStyleSetName(), "SourceControl.Actions.ChangeSettings")
-				);
-			}
-			else
-			{
-				InSection.AddMenuEntry(
-					FSourceControlCommands::Get().ConnectToSourceControl,
-					TAttribute<FText>(),
-					TAttribute<FText>(),
-					FSlateIcon(FEditorStyle::GetStyleSetName(), "SourceControl.Actions.Connect")
-				);
-			}
-		}));
-
-	Section.AddSeparator("SourceControlConnectionSeparator");
-
 	Section.AddMenuEntry(
 		FSourceControlCommands::Get().ViewChangelists,
 		TAttribute<FText>(),
@@ -163,6 +146,13 @@ TSharedRef<SWidget> FSourceControlMenuHelpers::GenerateSourceControlMenuContent(
 		FSlateIcon(FEditorStyle::GetStyleSetName(), "SourceControl.ChangelistsTab")
 	);
 
+	Section.AddMenuEntry(
+		FSourceControlCommands::Get().SubmitContent,
+		TAttribute<FText>(),
+		TAttribute<FText>(),
+		FSlateIcon(FEditorStyle::GetStyleSetName(), "SourceControl.Actions.Submit")
+	);
+
 	Section.AddMenuEntry(
 		FSourceControlCommands::Get().CheckOutModifiedFiles,
 		TAttribute<FText>(),
@@ -170,6 +160,29 @@ TSharedRef<SWidget> FSourceControlMenuHelpers::GenerateSourceControlMenuContent(
 		FSlateIcon(FEditorStyle::GetStyleSetName(), "SourceControl.Actions.CheckOut")
 	);
 
+	Section.AddDynamicEntry("ConnectToSourceControl", FNewToolMenuSectionDelegate::CreateLambda([](FToolMenuSection& InSection)
+	{
+		ISourceControlModule& SourceControlModule = ISourceControlModule::Get();
+		if (ISourceControlModule::Get().IsEnabled() && ISourceControlModule::Get().GetProvider().IsAvailable())
+		{
+			InSection.AddMenuEntry(
+				FSourceControlCommands::Get().ChangeSourceControlSettings,
+				TAttribute<FText>(),
+				TAttribute<FText>(),
+				FSlateIcon(FEditorStyle::GetStyleSetName(), "SourceControl.Actions.ChangeSettings")
+			);
+		}
+		else
+		{
+			InSection.AddMenuEntry(
+				FSourceControlCommands::Get().ConnectToSourceControl,
+				TAttribute<FText>(),
+				TAttribute<FText>(),
+				FSlateIcon(FEditorStyle::GetStyleSetName(), "SourceControl.Actions.Connect")
+			);
+		}
+	}));
+
 	return UToolMenus::Get()->GenerateWidget("StatusBar.ToolBar.SourceControl", FToolMenuContext(FSourceControlCommands::ActionList));
 }
 
diff --git a/Engine/Source/Editor/StatusBar/Private/SourceControlMenuHelpers.h b/Engine/Source/Editor/StatusBar/Private/SourceControlMenuHelpers.h
--- a/Engine/Source/Editor/StatusBar/Private/SourceControlMenuHelpers.h
+++ b/Engine/Source/Editor/StatusBar/Private/SourceControlMenuHelpers.h
@@ -34,6 +34,7 @@ public:
 	TSharedPtr< FUICommandInfo > ConnectToSourceControl;
 	TSharedPtr< FUICommandInfo > ChangeSourceControlSettings;
 	TSharedPtr< FUICommandInfo > ViewChangelists;
+	TSharedPtr< FUICommandInfo > SubmitContent;
 	TSharedPtr< FUICommandInfo > CheckOutModifiedFiles;
 
 	static TSharedRef<FUICommandList> ActionList;
