Remove editor throttle from status bar progress animations. They cause issues when the progress bar is active but the parent tab is not. While a fix exists where we stop throttling when the active tab changes, this throttle is happening without user interaction which voliates our rules on when we throttle.

diff --git a/Engine/Source/Editor/StatusBar/Private/SStatusBar.cpp b/Engine/Source/Editor/StatusBar/Private/SStatusBar.cpp
--- a/Engine/Source/Editor/StatusBar/Private/SStatusBar.cpp
+++ b/Engine/Source/Editor/StatusBar/Private/SStatusBar.cpp
@@ -509,7 +509,6 @@ public:
 
 				if (!OpenCloseTimer.IsValid())
 				{
-					AnimationThrottle = FSlateThrottleManager::Get().EnterResponsiveMode();
 					OpenCloseTimer = RegisterActiveTimer(0.0f, FWidgetActiveTimerDelegate::CreateSP(SharedThis(this), &SStatusBarProgressArea::UpdateProgressAnimation));
 				}
 			}
@@ -527,7 +526,6 @@ public:
 
 			if (!OpenCloseTimer.IsValid())
 			{
-				AnimationThrottle = FSlateThrottleManager::Get().EnterResponsiveMode();
 				OpenCloseTimer = RegisterActiveTimer(0.0f, FWidgetActiveTimerDelegate::CreateSP(this, &SStatusBarProgressArea::UpdateProgressAnimation));
 			}
 
@@ -549,7 +547,6 @@ private:
 				SetVisibility(EVisibility::Collapsed);
 			}
 
-			FSlateThrottleManager::Get().LeaveResponsiveMode(AnimationThrottle);
 			OpenCloseTimer.Reset();
 			return EActiveTimerReturnType::Stop;
 		}
