Added a new function FocusOutputLog that determines whether to open an output log tab or output log drawer. Fixed focus issue that made drawers close immediately when focus change is window activated.

diff --git a/Engine/Source/Developer/OutputLog/OutputLog.Build.cs b/Engine/Source/Developer/OutputLog/OutputLog.Build.cs
--- a/Engine/Source/Developer/OutputLog/OutputLog.Build.cs
+++ b/Engine/Source/Developer/OutputLog/OutputLog.Build.cs
@@ -1,32 +1,34 @@
-// Copyright Epic Games, Inc. All Rights Reserved.
-
-using UnrealBuildTool;
-
-public class OutputLog : ModuleRules
-{
-	public OutputLog(ReadOnlyTargetRules Target) : base(Target)
-	{
-		PrivateDependencyModuleNames.AddRange(
-			new string[] {
-				"Core",
-				"CoreUObject", // @todo Mac: for some reason it's needed to link in debug on Mac
-                "InputCore",
-				"Slate",
-				"SlateCore",
-                "TargetPlatform",
-                "DesktopPlatform",
-				"ToolWidgets"
-			}
-		);
-
-		if (Target.bBuildEditor)
-        {
-            PrivateDependencyModuleNames.AddRange(
-				new string[] {
-						"EditorFramework",
-                        "UnrealEd"
-				}
-			);
-        }
-    }
-}
+// Copyright Epic Games, Inc. All Rights Reserved.
+
+using UnrealBuildTool;
+
+public class OutputLog : ModuleRules
+{
+	public OutputLog(ReadOnlyTargetRules Target) : base(Target)
+	{
+		PrivateDependencyModuleNames.AddRange(
+			new string[] {
+				"Core",
+				"CoreUObject", // @todo Mac: for some reason it's needed to link in debug on Mac
+                "InputCore",
+				"Slate",
+				"SlateCore",
+                "TargetPlatform",
+                "DesktopPlatform",
+				"ToolWidgets", 
+				"Engine", 
+			}
+		);
+
+		if (Target.bBuildEditor)
+        {
+            PrivateDependencyModuleNames.AddRange(
+				new string[] {
+					"EditorFramework",
+					"UnrealEd",
+					"StatusBar",
+				}
+			);
+        }
+    }
+}
diff --git a/Engine/Source/Developer/OutputLog/Private/OutputLogModule.cpp b/Engine/Source/Developer/OutputLog/Private/OutputLogModule.cpp
--- a/Engine/Source/Developer/OutputLog/Private/OutputLogModule.cpp
+++ b/Engine/Source/Developer/OutputLog/Private/OutputLogModule.cpp
@@ -28,6 +28,10 @@
 #include "Internationalization/Internationalization.h"
 #include "OutputLogStyle.h"
 
+#if WITH_EDITOR
+#include "StatusBarSubsystem.h"
+#endif
+
 IMPLEMENT_MODULE(FOutputLogModule, OutputLog);
 
 namespace OutputLogModule
@@ -287,7 +291,7 @@
 
 	if (bClearOnPIEEnabled)
 	{
-		if(TSharedPtr<SOutputLog> OutputLogPinned = OutputLog.Pin())
+		if (TSharedPtr<SOutputLog> OutputLogPinned = OutputLog.Pin())
 		{
 			if (OutputLogPinned->CanClearLog())
 			{
@@ -322,6 +326,40 @@
 	return OutputLog.Pin();
 }
 
+void FOutputLogModule::FocusOutputLog()
+{
+	// 1. Output log tab is open but not active. 
+	if (OutputLog.IsValid()) 
+	{
+		OutputLogTab.Pin()->DrawAttention(); 
+	}
+#if WITH_EDITOR
+	// 2. Output log tab isn't open and the window directly behind the notification window has a status bar, then open Output Log Drawer. 
+	else if (GEditor->GetEditorSubsystem<UStatusBarSubsystem>()->ActiveWindowBehindNotificationHasStatusBar())
+	{
+		TSharedRef<SWindow> ParentWindow = FSlateApplication::Get().GetActiveTopLevelRegularWindow().ToSharedRef(); 
+
+		// try toggle the console to open the Output Log Drawer
+		if (GEditor->GetEditorSubsystem<UStatusBarSubsystem>()->ToggleDebugConsole(ParentWindow, true))
+		{
+			OutputLogDrawer.Pin()->FocusConsoleCommandBox(); 
+		}
+		// if unable to open the drawer, invoke a new Output Log tab. 
+		else 
+		{ 
+			OpenOutputLog(); 
+			OutputLogTab.Pin()->DrawAttention();
+		}
+	}
+#endif
+	// 3. The parent window has no status bar, then invoke a new Output Log tab.  
+	else
+	{
+		OpenOutputLog(); 
+		OutputLogTab.Pin()->DrawAttention(); 
+	}
+}
+
 void FOutputLogModule::UpdateOutputLogFilter(const TArray<FName>& CategoriesToShow, TOptional<bool> bShowErrors, TOptional<bool> bShowWarnings, TOptional<bool> bShowLogs)
 {
 	if (TSharedPtr<SOutputLog> SharedOutputLog = OutputLog.Pin())
@@ -334,3 +372,8 @@
 {
 	FGlobalTabmanager::Get()->TryInvokeTab(OutputLogModule::OutputLogTabName);
 }
+
+bool FOutputLogModule::ShouldCycleToOutputLogDrawer() const 
+{
+	return GetDefault<UOutputLogSettings>()->bCycleToOutputLogDrawer; 
+}
diff --git a/Engine/Source/Developer/OutputLog/Public/OutputLogModule.h b/Engine/Source/Developer/OutputLog/Public/OutputLogModule.h
--- a/Engine/Source/Developer/OutputLog/Public/OutputLogModule.h
+++ b/Engine/Source/Developer/OutputLog/Public/OutputLogModule.h
@@ -61,6 +61,10 @@ public:
 	virtual void FocusOutputLogConsoleBox(const TSharedRef<SWidget> OutputLogToFocus);
 
 	virtual const TSharedPtr<SWidget> GetOutputLog() const;
+
+	/** Opens and focuses on the Output Log Drawer if the status bar exists, otherwise opens and focuses on the Output Log Tab. */
+	void OUTPUTLOG_API FocusOutputLog();
+
 	const TSharedPtr<SDockTab> GetOutputLogTab() const { return OutputLogTab.Pin(); }
 
 	/** Change the output log's filter. If CategoriesToShow is empty, all categories will be shown. */
@@ -68,6 +72,10 @@ public:
 	/** Opens the output log tab, or brings it to front if it's already open */
 	void OUTPUTLOG_API OpenOutputLog() const;
 
+	/** Returns the value of bCycleToOutputLogDrawer from the module OutputLogSettings. This function helps StatusBar to access properties in OutputLogSetting 
+	through OutputLogModule to avoid dependencies. */
+	virtual bool ShouldCycleToOutputLogDrawer() const;
+
 private:
 	TSharedRef<SDockTab> SpawnOutputLogTab(const FSpawnTabArgs& Args);
 	TSharedRef<SDockTab> SpawnDeviceOutputLogTab(const FSpawnTabArgs& Args);
diff --git a/Engine/Source/Developer/ToolWidgets/Private/SWidgetDrawer.cpp b/Engine/Source/Developer/ToolWidgets/Private/SWidgetDrawer.cpp
--- a/Engine/Source/Developer/ToolWidgets/Private/SWidgetDrawer.cpp
+++ b/Engine/Source/Developer/ToolWidgets/Private/SWidgetDrawer.cpp
@@ -353,7 +353,7 @@ void SWidgetDrawer::Construct(const FArguments& InArgs, FName InStatusBarName)
 void SWidgetDrawer::OnGlobalFocusChanging(const FFocusEvent& FocusEvent, const FWeakWidgetPath& OldFocusedWidgetPath, const TSharedPtr<SWidget>& OldFocusedWidget, const FWidgetPath& NewFocusedWidgetPath, const TSharedPtr<SWidget>& NewFocusedWidget)
 {
 	// Sometimes when dismissing focus can change which will trigger this again
-	static bool bIsRentrant = false;
+	static bool bIsRentrant = false; 
 
 	if(!bIsRentrant)
 	{
@@ -388,7 +388,10 @@ void SWidgetDrawer::OnGlobalFocusChanging(const FFocusEvent& FocusEvent, const F
 					bShouldDismiss = true;
 				}
 			}
-			else
+			// When the focus change is initiated by the window, don't dismiss the drawer. 
+			// Scenario: when users try to open the Output Log Drawer via the "View Output Log" hyperlink, and the Output Log Drawer 
+			// wasn't in focus already, the window will dismiss the drawer immediately after it's opened due to focus change. 
+			else if (FocusEvent.GetCause() != EFocusCause::WindowActivate)
 			{
 				bShouldDismiss = true;
 			}
diff --git a/Engine/Source/Editor/StatusBar/Private/StatusBarSubsystem.cpp b/Engine/Source/Editor/StatusBar/Private/StatusBarSubsystem.cpp
--- a/Engine/Source/Editor/StatusBar/Private/StatusBarSubsystem.cpp
+++ b/Engine/Source/Editor/StatusBar/Private/StatusBarSubsystem.cpp
@@ -259,13 +259,14 @@
 {
 	bool bToggledSuccessfully = false;
 
-	FOutputLogModule& OutputLogModule = FOutputLogModule::Get();
+	FOutputLogModule& OutputLogModule = FModuleManager::Get().LoadModuleChecked<FOutputLogModule>("OutputLog");
 
+
 	// Get the global output log tab if it exists. If it exists and is in the same editor as the status bar we'll focus that instead
 	TSharedPtr<SDockTab> MainOutputLogTab = OutputLogModule.GetOutputLogTab();
 	
-	const bool bCycleToOutputLogDrawer = GetDefault<UOutputLogSettings>()->bCycleToOutputLogDrawer || bAlwaysToggleDrawer;
-
+	const bool bCycleToOutputLogDrawer = OutputLogModule.ShouldCycleToOutputLogDrawer() || bAlwaysToggleDrawer;
+	
 	for (auto StatusBar : StatusBars)
 	{
 		FStatusBarData& SBData = StatusBar.Value;
@@ -339,6 +340,12 @@
 {
 	TSharedPtr<SWindow> ParentWindow = UE::StatusBarSubsystem::Private::FindParentWindow();
 
+	if (ParentWindow.IsValid() && ParentWindow->GetType() == EWindowType::Notification)
+	{
+		// Get the parent window directly behind the notification. 
+		ParentWindow = FSlateApplication::Get().GetActiveTopLevelRegularWindow(); 
+	}
+	
 	if (ParentWindow.IsValid() && ParentWindow->GetType() == EWindowType::Normal)
 	{
 		return ToggleDebugConsole(ParentWindow.ToSharedRef(), true);
@@ -444,7 +451,7 @@
 
 	StatusBar->RegisterDrawer(MoveTemp(ContentBrowserDrawer));
 
-	FOutputLogModule& OutputLogModule = FOutputLogModule::Get();
+	FOutputLogModule& OutputLogModule = FModuleManager::Get().LoadModuleChecked<FOutputLogModule>("OutputLog");
 
 	TWeakPtr<SStatusBar> StatusBarWeakPtr = StatusBar;
 
@@ -518,6 +525,31 @@
 	return false;
 }
 
+// This function is identical to UStatusBarSubsystem::ActiveWindowHasStatusBar(), except the variable ParentWindow stores the topmost
+// Regular window rather than the direct parent window. This function is only used when users click on the hyperlink "Show Output Log" on 
+// the notification window. 
+bool UStatusBarSubsystem::ActiveWindowBehindNotificationHasStatusBar()
+{
+	TSharedPtr<SWindow> ParentWindow = FSlateApplication::Get().GetActiveTopLevelRegularWindow(); 
+
+	// Same code as ActiveWindowHasStatusBar() from here: 
+	if (ParentWindow.IsValid() && ParentWindow->GetType() == EWindowType::Normal)
+	{
+		for (const TPair<FName, FStatusBarData>& StatusBar : StatusBars)
+		{
+			if (TSharedPtr<SStatusBar> StatusBarPinned = StatusBar.Value.StatusBarWidget.Pin())
+			{
+				TSharedPtr<SDockTab> ParentTab = StatusBarPinned->GetParentTab();
+				if (ParentTab && ParentTab->IsForeground() && ParentTab->GetParentWindow() == ParentWindow)
+				{
+					return true;
+				}
+			}
+		}
+	}
+	return false;
+}
+
 void UStatusBarSubsystem::RegisterDrawer(FName StatusBarName, FWidgetDrawerConfig&& Drawer, int32 SlotIndex)
 {
 	if (TSharedPtr<SStatusBar> StatusBar = GetStatusBar(StatusBarName))
@@ -767,9 +799,11 @@
 
 TSharedRef<SWidget> UStatusBarSubsystem::OnGetOutputLog()
 {
+	FOutputLogModule& OutputLogModule = FModuleManager::Get().LoadModuleChecked<FOutputLogModule>("OutputLog"); 
+
 	if (!StatusBarOutputLog)
 	{
-		StatusBarOutputLog = FOutputLogModule::Get().MakeOutputLogDrawerWidget(FSimpleDelegate::CreateUObject(this, &UStatusBarSubsystem::OnDebugConsoleDrawerClosed));
+		StatusBarOutputLog = OutputLogModule.MakeOutputLogDrawerWidget(FSimpleDelegate::CreateUObject(this, &UStatusBarSubsystem::OnDebugConsoleDrawerClosed));
 	}
 
 	return StatusBarOutputLog.ToSharedRef();
@@ -788,8 +822,8 @@
 			}
 		}
 	}
-
-	FOutputLogModule::Get().FocusOutputLogConsoleBox(StatusBarOutputLog.ToSharedRef());
+	FOutputLogModule& OutputLogModule = FModuleManager::Get().LoadModuleChecked<FOutputLogModule>("OutputLog");
+	OutputLogModule.FocusOutputLogConsoleBox(StatusBarOutputLog.ToSharedRef());
 }
 
 void UStatusBarSubsystem::OnOutputLogDismised(const TSharedPtr<SWidget>& NewlyFocusedWidget)
diff --git a/Engine/Source/Editor/StatusBar/Public/StatusBarSubsystem.h b/Engine/Source/Editor/StatusBar/Public/StatusBarSubsystem.h
--- a/Engine/Source/Editor/StatusBar/Public/StatusBarSubsystem.h
+++ b/Engine/Source/Editor/StatusBar/Public/StatusBarSubsystem.h
@@ -115,6 +115,11 @@ public:
 	 */
 	bool ActiveWindowHasStatusBar() const;
 
+	/**
+	 * @return true if a status bar was found for the active window immediately behind a Notification Window
+	*/
+	bool ActiveWindowBehindNotificationHasStatusBar();
+
 	/**
 	 * Creates a new instance of a status bar widget
 	 *
diff --git a/Engine/Source/Editor/StatusBar/StatusBar.Build.cs b/Engine/Source/Editor/StatusBar/StatusBar.Build.cs
--- a/Engine/Source/Editor/StatusBar/StatusBar.Build.cs
+++ b/Engine/Source/Editor/StatusBar/StatusBar.Build.cs
@@ -18,7 +18,6 @@ public class StatusBar : ModuleRules
 				"EditorSubsystem",
 				"Engine",
 				"InputCore",
-				"OutputLog",
 				"Slate",
 				"SlateCore",
 				"InputCore",
@@ -32,11 +31,15 @@ public class StatusBar : ModuleRules
 		DynamicallyLoadedModuleNames.AddRange(
 			new string[] {
 				"MainFrame",
+				"OutputLog",
+				"ContentBrowser", 
 			});
 
 		PrivateIncludePathModuleNames.AddRange(
 			new string[] {
 				"MainFrame",
+				"OutputLog",
+				"ContentBrowser",
 			});
 	}
 }
diff --git a/Engine/Source/Editor/UATHelper/UATHelper.Build.cs b/Engine/Source/Editor/UATHelper/UATHelper.Build.cs
--- a/Engine/Source/Editor/UATHelper/UATHelper.Build.cs
+++ b/Engine/Source/Editor/UATHelper/UATHelper.Build.cs
@@ -18,6 +18,7 @@ namespace UnrealBuildTool.Rules
 					"EditorFramework",
                     "UnrealEd",
 					"Analytics",
+					"OutputLog", 
 			    }
 			);
 		}
diff --git a/Engine/Source/Editor/UATHelper/UATHelperModule.cpp b/Engine/Source/Editor/UATHelper/UATHelperModule.cpp
--- a/Engine/Source/Editor/UATHelper/UATHelperModule.cpp
+++ b/Engine/Source/Editor/UATHelper/UATHelperModule.cpp
@@ -36,6 +36,7 @@
 #include "AnalyticsEventAttribute.h"
 
 #include "ShaderCompiler.h"
+#include "Developer/OutputLog/Public/OutputLogModule.h"
 
 #define LOCTEXT_NAMESPACE "UATHelper"
 
@@ -487,7 +488,8 @@ public:
 
 	static void HandleUatHyperlinkNavigate()
 	{
-		FGlobalTabmanager::Get()->TryInvokeTab(FName("OutputLog"));
+		FOutputLogModule& OutputLogModule = FOutputLogModule::Get();
+		OutputLogModule.FocusOutputLog();
 	}
 
 	static void HandleUatResultHyperlinkNavigate(FString ResultLocation)
