Adding a one-time popup for a user's industry for analytics. Off in internal builds, and can be suppressed either by turning off all analytics or specifically with bSuppressIndustryPopup on UAnalyticsPrivacySettings.

diff --git a/Engine/Source/Editor/StatusBar/Private/SOneTimeIndustryQuery.cpp b/Engine/Source/Editor/StatusBar/Private/SOneTimeIndustryQuery.cpp
new file mode 100644
--- /dev/null
+++ b/Engine/Source/Editor/StatusBar/Private/SOneTimeIndustryQuery.cpp
@@ -0,0 +1,236 @@
+// Copyright Epic Games, Inc. All Rights Reserved.
+
+#include "SOneTimeIndustryQuery.h"
+#include "Widgets/Notifications/SNotificationBackground.h"
+#include "Widgets/SWindow.h"
+#include "Widgets/Layout/SBox.h"
+#include "Widgets/Layout/SBorder.h"
+#include "Widgets/Text/STextBlock.h"
+#include "Widgets/Input/SButton.h"
+#include "Widgets/Images/SImage.h"
+#include "Widgets/Input/SComboBox.h"
+#include "SPrimaryButton.h"
+#include "EngineAnalytics.h"
+#include "Misc/EngineBuildSettings.h"
+#include "Editor/EditorEngine.h"
+#include "Analytics/AnalyticsPrivacySettings.h"
+
+#define LOCTEXT_NAMESPACE "OneTimeIndustryQuery"
+
+TWeakPtr<SOneTimeIndustryQuery> SOneTimeIndustryQuery::ActiveNotification;
+TWeakPtr<SWindow> SOneTimeIndustryQuery::ParentWindow;
+
+const FName SOneTimeIndustryQuery::GamesIndustryName = FName("Games");
+const FName SOneTimeIndustryQuery::FilmIndustryName = FName("FilmAndTV");
+const FName SOneTimeIndustryQuery::ArchitectureIndustryName = FName("Architecture");
+const FName SOneTimeIndustryQuery::AutoIndustryName = FName("AutomotiveAndManufacturing");
+const FName SOneTimeIndustryQuery::BroadcastIndustryName = FName("BroadcastingAndLiveEvents");
+const FName SOneTimeIndustryQuery::AdIndustryName = FName("AdvertisingAndMarketing");
+const FName SOneTimeIndustryQuery::SimulationIndustryName = FName("SimulationAndTraining");
+const FName SOneTimeIndustryQuery::OtherIndustryName = FName("Other");
+
+void SOneTimeIndustryQuery::Show(TSharedPtr<SWindow> InParentWindow)
+{
+	if (!ActiveNotification.IsValid())
+	{
+		TSharedRef<SOneTimeIndustryQuery> ActiveNotificationRef =
+			SNew(SOneTimeIndustryQuery);
+
+		ActiveNotification = ActiveNotificationRef;
+		ParentWindow = InParentWindow;
+		InParentWindow->AddOverlaySlot()
+			.VAlign(VAlign_Bottom)
+			.HAlign(HAlign_Right)
+			.Padding(FMargin(20.0f, 20.0f, 10.0f, 50.f))
+			[
+				ActiveNotificationRef
+			];
+	}
+}
+
+void SOneTimeIndustryQuery::Dismiss()
+{
+	TSharedPtr<SOneTimeIndustryQuery> ActiveNotificationPin = ActiveNotification.Pin();
+	if (ParentWindow.IsValid() && ActiveNotificationPin.IsValid())
+	{
+		ParentWindow.Pin()->RemoveOverlaySlot(ActiveNotificationPin.ToSharedRef());
+	}
+
+	ParentWindow.Reset();
+
+	ActiveNotification.Reset();
+}
+
+const bool SOneTimeIndustryQuery::ShouldShowIndustryQuery()
+{
+	return !FEngineBuildSettings::IsInternalBuild()  
+		&& GEngine->AreEditorAnalyticsEnabled()
+		&& !GetDefault<UAnalyticsPrivacySettings>()->bSuppressIndustryPopup;
+}
+
+void SOneTimeIndustryQuery::Construct(const FArguments& InArgs)
+{
+	IndustryComboList.Add(MakeShared<FName>(SOneTimeIndustryQuery::GamesIndustryName));
+	IndustryComboList.Add(MakeShared<FName>(SOneTimeIndustryQuery::FilmIndustryName));
+	IndustryComboList.Add(MakeShared<FName>(SOneTimeIndustryQuery::ArchitectureIndustryName));
+	IndustryComboList.Add(MakeShared<FName>(SOneTimeIndustryQuery::AutoIndustryName));
+	IndustryComboList.Add(MakeShared<FName>(SOneTimeIndustryQuery::BroadcastIndustryName));
+	IndustryComboList.Add(MakeShared<FName>(SOneTimeIndustryQuery::AdIndustryName));
+	IndustryComboList.Add(MakeShared<FName>(SOneTimeIndustryQuery::SimulationIndustryName));
+	IndustryComboList.Add(MakeShared<FName>(SOneTimeIndustryQuery::OtherIndustryName));
+
+	ChildSlot
+		[
+			SNew(SBox)
+			.WidthOverride(375.0f)
+			.HeightOverride(128.0f)
+			[
+				SNew(SNotificationBackground)
+				.Padding(FMargin(8, 8))
+				[
+					SNew(SHorizontalBox)
+					+ SHorizontalBox::Slot()
+					.FillWidth(1.0f)
+					.Padding(20.0f, 15.0f, 0.0f, 0.0f)
+					[
+						SNew(SVerticalBox)
+
+						+ SVerticalBox::Slot()
+						.AutoHeight()
+						.Padding(0.0f, 0.0f, 0.0f, 0.0f)
+						[
+							SNew(STextBlock)
+							.Font(FAppStyle::Get().GetFontStyle("NotificationList.FontBold"))
+							.Text(LOCTEXT("IndustryQuestionDesc", "Help us improve Unreal Engine by \ntelling us your industry."))
+							.ColorAndOpacity(FStyleColors::ForegroundHover)
+						]
+
+						+ SVerticalBox::Slot()
+						.VAlign(VAlign_Center)
+						.FillHeight(1.0f)
+						.Padding(0.0f, 12.0f, 0.0f, 12.0f)
+						[
+							SNew(SHorizontalBox)
+
+							+ SHorizontalBox::Slot()
+							.FillWidth(1.0f)
+							[
+								SNew(SComboBox<TSharedPtr<FName>>)
+								.OptionsSource(&IndustryComboList)
+								.OnGenerateWidget(this, &SOneTimeIndustryQuery::GenerateIndustryComboItem)
+								.OnSelectionChanged(this, &SOneTimeIndustryQuery::HandleIndustryComboChanged)
+								[
+									SNew(STextBlock)
+									.Text(this, &SOneTimeIndustryQuery::GetUserSetIndustryNameText)
+								]
+							]
+
+							+ SHorizontalBox::Slot()
+							.AutoWidth()
+							.Padding(15.0f, 0.0f, 0.0f, 0.0f)
+							[
+								SNew(SPrimaryButton)
+								.Text(LOCTEXT("IndustryFeedbackSubmit", "Submit"))
+								.IsEnabled(this, &SOneTimeIndustryQuery::CanSubmitIndustryInfo)
+								.OnClicked(this, &SOneTimeIndustryQuery::OnSubmit)
+							]
+						]
+					]
+
+					+ SHorizontalBox::Slot()
+					.Padding(0.0f, 0.0f, 0.0f, 0.0f)
+					.HAlign(HAlign_Right)
+					.VAlign(VAlign_Top)
+					.AutoWidth()
+					[
+						SNew(SButton)
+						.ButtonStyle(FAppStyle::Get(), "SimpleButton")
+						.OnClicked_Lambda([]() { SOneTimeIndustryQuery::Dismiss(); return FReply::Handled(); })
+						[
+							SNew(SImage)
+							.Image(FAppStyle::Get().GetBrush("Icons.X"))
+							.ColorAndOpacity(FSlateColor::UseForeground())
+						]
+					]
+				]
+			]
+		];
+
+}
+
+bool SOneTimeIndustryQuery::CanSubmitIndustryInfo() const
+{
+	return UserSetIndustry.IsSet();
+}
+
+void SOneTimeIndustryQuery::HandleIndustryComboChanged(TSharedPtr<FName> Item, ESelectInfo::Type SelectInfo)
+{
+	FSlateApplication::Get().DismissAllMenus();
+
+	UserSetIndustry = *Item.Get();
+}
+
+TSharedRef<SWidget> SOneTimeIndustryQuery::GenerateIndustryComboItem(TSharedPtr<FName> InItem) const
+{
+	return SNew(STextBlock).Text(GetIndustryNameText(*InItem.Get()));
+}
+
+FText SOneTimeIndustryQuery::GetUserSetIndustryNameText() const
+{
+	if (!UserSetIndustry.IsSet())
+	{
+		return LOCTEXT("SelectYourIndustry", "Select Your Industry");
+	}
+	return GetIndustryNameText(UserSetIndustry.GetValue());
+}
+
+FText SOneTimeIndustryQuery::GetIndustryNameText(FName IndustryName) const
+{
+	if (IndustryName == SOneTimeIndustryQuery::GamesIndustryName)
+	{
+		return LOCTEXT("GameDevelopment", "Game Development");
+	}
+	else if (IndustryName == SOneTimeIndustryQuery::FilmIndustryName)
+	{
+		return LOCTEXT("FilmAndTV", "Film & TV");
+	}
+	else if (IndustryName == SOneTimeIndustryQuery::ArchitectureIndustryName)
+	{
+		return LOCTEXT("Architecture", "Architecture");
+	}
+	else if (IndustryName == SOneTimeIndustryQuery::AutoIndustryName)
+	{
+		return LOCTEXT("AutomotiveAndManufacturing", "Automotive & Manufacturing");
+	}
+	else if (IndustryName == SOneTimeIndustryQuery::BroadcastIndustryName)
+	{
+		return LOCTEXT("BroadcastingAndLiveEvents", "Broadcasting & Live Events");
+	}
+	else if (IndustryName == SOneTimeIndustryQuery::AdIndustryName)
+	{
+		return LOCTEXT("AdvertisingAndMarketing", "Advertising & Marketing");
+	}
+	else if (IndustryName == SOneTimeIndustryQuery::SimulationIndustryName)
+	{
+		return LOCTEXT("SimulationAndTraining", "Simulation & Training");
+	}
+	else if (IndustryName == SOneTimeIndustryQuery::OtherIndustryName)
+	{
+		return LOCTEXT("Other", "Other");
+	}
+	return FText::GetEmpty();
+}
+
+FReply SOneTimeIndustryQuery::OnSubmit()
+{
+	if (FEngineAnalytics::IsAvailable())
+	{
+		TArray< FAnalyticsEventAttribute > IndustryAttribs;
+		IndustryAttribs.Add(FAnalyticsEventAttribute(TEXT("Industry"), UserSetIndustry.GetValue().ToString()));
+		FEngineAnalytics::GetProvider().RecordEvent(TEXT("Editor.UserIndustry"), IndustryAttribs);
+	}
+	SOneTimeIndustryQuery::Dismiss();
+	return FReply::Handled();
+}
+
+#undef LOCTEXT_NAMESPACE
diff --git a/Engine/Source/Editor/StatusBar/Private/SOneTimeIndustryQuery.h b/Engine/Source/Editor/StatusBar/Private/SOneTimeIndustryQuery.h
new file mode 100644
--- /dev/null
+++ b/Engine/Source/Editor/StatusBar/Private/SOneTimeIndustryQuery.h
@@ -0,0 +1,56 @@
+// Copyright Epic Games, Inc. All Rights Reserved.
+
+#pragma once
+#include "Styling/StyleColors.h"
+#include "Brushes/SlateRoundedBoxBrush.h"
+#include "Templates/SharedPointer.h"
+#include "Widgets/SCompoundWidget.h"
+
+class SWindow;
+
+class SOneTimeIndustryQuery : public SCompoundWidget
+{
+	SLATE_BEGIN_ARGS(SOneTimeIndustryQuery)
+	{}
+	SLATE_END_ARGS()
+
+	SOneTimeIndustryQuery()
+	: NewBadgeBrush(FStyleColors::Success)
+	, KeybindBackgroundBrush(FLinearColor::Transparent, 6.0f, FStyleColors::ForegroundHover, 1.5f)
+	{}
+
+	static void Show(TSharedPtr<SWindow> InParentWindow);
+
+	static void Dismiss();
+
+	static const bool ShouldShowIndustryQuery();
+
+	void Construct(const FArguments& InArgs);
+
+private:
+	bool CanSubmitIndustryInfo() const;
+	void HandleIndustryComboChanged(TSharedPtr<FName> Item, ESelectInfo::Type SelectInfo);
+	TSharedRef<SWidget> GenerateIndustryComboItem(TSharedPtr<FName> InItem) const;
+	FText GetUserSetIndustryNameText() const;
+	FText GetIndustryNameText(FName IndustryName) const;
+	FReply OnSubmit();
+
+private:
+	FSlateRoundedBoxBrush NewBadgeBrush;
+	FSlateRoundedBoxBrush KeybindBackgroundBrush;
+	static TWeakPtr<SOneTimeIndustryQuery> ActiveNotification;
+	static TWeakPtr<SWindow> ParentWindow;
+	TArray<TSharedPtr<FName>> IndustryComboList;
+	TOptional<FName> UserSetIndustry;
+
+	static const FName GamesIndustryName;
+	static const FName FilmIndustryName;
+	static const FName ArchitectureIndustryName;
+	static const FName AutoIndustryName;
+	static const FName BroadcastIndustryName;
+	static const FName AdIndustryName;
+	static const FName SimulationIndustryName;
+	static const FName OtherIndustryName;
+};
+
+
diff --git a/Engine/Source/Editor/StatusBar/Private/StatusBarSubsystem.cpp b/Engine/Source/Editor/StatusBar/Private/StatusBarSubsystem.cpp
--- a/Engine/Source/Editor/StatusBar/Private/StatusBarSubsystem.cpp
+++ b/Engine/Source/Editor/StatusBar/Private/StatusBarSubsystem.cpp
@@ -25,6 +25,7 @@
 #include "WidgetDrawerConfig.h"
 #include "Widgets/Input/SMultiLineEditableTextBox.h"
 #include "OutputLog/Public/OutputLogSettings.h"
+#include "SOneTimeIndustryQuery.h"
 
 #define LOCTEXT_NAMESPACE "StatusBar"
 
@@ -235,10 +236,12 @@ void UStatusBarSubsystem::Initialize(FSubsystemCollectionBase& Collection)
 	if (MainFrameModule.IsWindowInitialized())
 	{
 		CreateAndShowNewUserTipIfNeeded(MainFrameModule.GetParentWindow(), false);
+		CreateAndShowOneTimeIndustryQueryIfNeeded(MainFrameModule.GetParentWindow(), false);
 	}
 	else
 	{
 		MainFrameModule.OnMainFrameCreationFinished().AddUObject(this, &UStatusBarSubsystem::CreateAndShowNewUserTipIfNeeded);
+		MainFrameModule.OnMainFrameCreationFinished().AddUObject(this, &UStatusBarSubsystem::CreateAndShowOneTimeIndustryQueryIfNeeded);
 	}
 
 
@@ -697,41 +700,92 @@ void UStatusBarSubsystem::CreateAndShowNewUserTipIfNeeded(TSharedPtr<SWindow> Pa
 {
 	if(!bIsNewProjectDialog)
 	{
-		const FString StoreId = TEXT("Epic Games");
-		const FString SectionName = TEXT("Unreal Engine/Editor");
-		const FString KeyName = TEXT("LaunchTipShown");
+		FString CurrentState = GetNewUserTipState();
 
-		const FString FallbackIniLocation = TEXT("/Script/UnrealEd.EditorSettings");
-		const FString FallbackIniKey = TEXT("LaunchTipShownFallback");
-
-		// Its important that this new user message does not appear after the first launch so we store it in a more permanent place
-		FString CurrentState = TEXT("0");
-		if (!FPlatformMisc::GetStoredValue(StoreId, SectionName, KeyName, CurrentState))
+		if(CurrentState != TEXT("1"))
 		{
-			// As a fallback where the registry was not readable or writable, save a flag in the editor ini. This will be less permanent as the registry but will prevent 
-			// the notification from appearing on every launch
-			GConfig->GetString(*FallbackIniLocation, *FallbackIniKey, CurrentState, GEditorSettingsIni);
+			SNewUserTipNotification::Show(ParentWindow);
 
+			const FString StoreId = TEXT("Epic Games");
+			const FString SectionName = TEXT("Unreal Engine/Editor");
+			const FString KeyName = TEXT("LaunchTipShown");
 
+			const FString FallbackIniLocation = TEXT("/Script/UnrealEd.EditorSettings");
+			const FString FallbackIniKey = TEXT("LaunchTipShownFallback");
+			// Write that we've shown the notification
+			UStatusBarSubsystem::SetOneTimeStateWithFallback(StoreId, SectionName, KeyName, FallbackIniLocation, FallbackIniKey);
 		}
+	}
 
-		if(CurrentState != TEXT("1"))
+	// Ignore the if the main frame gets recreated this session
+	IMainFrameModule::Get().OnMainFrameCreationFinished().RemoveAll(this);
+}
+
+const FString UStatusBarSubsystem::GetNewUserTipState() const
+{
+	const FString StoreId = TEXT("Epic Games");
+	const FString SectionName = TEXT("Unreal Engine/Editor");
+	const FString KeyName = TEXT("LaunchTipShown");
+
+	const FString FallbackIniLocation = TEXT("/Script/UnrealEd.EditorSettings");
+	const FString FallbackIniKey = TEXT("LaunchTipShownFallback");
+
+	FString CurrentState = UStatusBarSubsystem::GetOneTimeStateWithFallback(StoreId, SectionName, KeyName, FallbackIniLocation, FallbackIniKey);
+	
+	return CurrentState;
+}
+
+void UStatusBarSubsystem::CreateAndShowOneTimeIndustryQueryIfNeeded(TSharedPtr<SWindow> ParentWindow, bool bIsNewProjectDialog)
+{
+	// Only show for external builds where editor analytics are on and the industry popup is not suppressed
+	if(SOneTimeIndustryQuery::ShouldShowIndustryQuery())
+	{
+		FString NewUserTipState = GetNewUserTipState();
+		if (!bIsNewProjectDialog && NewUserTipState == TEXT("1"))
 		{
-			SNewUserTipNotification::Show(ParentWindow);
+			const FString StoreId = TEXT("Epic Games");
+			const FString SectionName = TEXT("Unreal Engine/Editor");
+			const FString KeyName = TEXT("OneTimeIndustryQueryShown");
 
-			// Write that we've shown the notification
-			if (!FPlatformMisc::SetStoredValue(StoreId, SectionName, KeyName, TEXT("1")))
+			const FString FallbackIniLocation = TEXT("/Script/UnrealEd.EditorSettings");
+			const FString FallbackIniKey = TEXT("OneTimeIndustryQueryShownFallback");
+
+			FString CurrentState = UStatusBarSubsystem::GetOneTimeStateWithFallback(StoreId, SectionName, KeyName, FallbackIniLocation, FallbackIniKey);
+
+			if (CurrentState != TEXT("1"))
 			{
-				// Use fallback
-				GConfig->SetString(*FallbackIniLocation, *FallbackIniKey, TEXT("1"), GEditorSettingsIni);
+				SOneTimeIndustryQuery::Show(ParentWindow);
+				SetOneTimeStateWithFallback(StoreId, SectionName, KeyName, FallbackIniLocation, FallbackIniKey);
 			}
 		}
 	}
-
-	// Ignore the if the main frame gets recreated this session
+	// Ignore this if the main frame gets recreated this session
 	IMainFrameModule::Get().OnMainFrameCreationFinished().RemoveAll(this);
 }
 
+const FString UStatusBarSubsystem::GetOneTimeStateWithFallback(const FString StoreId, const FString SectionName, const FString KeyName, const FString FallbackIniLocation, const FString FallbackIniKey)
+{
+	// Its important that this new user message does not appear after the first launch so we store it in a more permanent place
+	FString CurrentState = TEXT("0");
+	if (!FPlatformMisc::GetStoredValue(StoreId, SectionName, KeyName, CurrentState))
+	{
+		// As a fallback where the registry was not readable or writable, save a flag in the editor ini. This will be less permanent as the registry but will prevent 
+		// the notification from appearing on every launch
+		GConfig->GetString(*FallbackIniLocation, *FallbackIniKey, CurrentState, GEditorSettingsIni);
+	}
+	return CurrentState;
+}
+
+void UStatusBarSubsystem::SetOneTimeStateWithFallback(const FString StoreId, const FString SectionName, const FString KeyName, const FString FallbackIniLocation, const FString FallbackIniKey)
+{
+	// Write that we've shown the notification
+	if (!FPlatformMisc::SetStoredValue(StoreId, SectionName, KeyName, TEXT("1")))
+	{
+		// Use fallback
+		GConfig->SetString(*FallbackIniLocation, *FallbackIniKey, TEXT("1"), GEditorSettingsIni);
+	}
+}
+
 TSharedPtr<SStatusBar> UStatusBarSubsystem::GetStatusBar(FName StatusBarName) const
 {
 	return StatusBars.FindRef(StatusBarName).StatusBarWidget.Pin();
diff --git a/Engine/Source/Editor/StatusBar/Public/StatusBarSubsystem.h b/Engine/Source/Editor/StatusBar/Public/StatusBarSubsystem.h
--- a/Engine/Source/Editor/StatusBar/Public/StatusBarSubsystem.h
+++ b/Engine/Source/Editor/StatusBar/Public/StatusBarSubsystem.h
@@ -165,6 +165,10 @@ private:
 	void OnDebugConsoleClosed(TWeakPtr<SStatusBar> OwningStatusBar);
 	void CreateContentBrowserIfNeeded();
 	void CreateAndShowNewUserTipIfNeeded(TSharedPtr<SWindow> ParentWindow, bool bIsNewProjectDialog);
+	const FString GetNewUserTipState() const;
+	void CreateAndShowOneTimeIndustryQueryIfNeeded(TSharedPtr<SWindow> ParentWindow, bool bIsNewProjectDialog);
+	static const FString GetOneTimeStateWithFallback(const FString StoreId, const FString SectionName, const FString KeyName, const FString FallbackIniLocation, const FString FallbackIniKey);
+	static void SetOneTimeStateWithFallback(const FString StoreId, const FString SectionName, const FString KeyName, const FString FallbackIniLocation, const FString FallbackIniKey);
 
 	TSharedPtr<SStatusBar> GetStatusBar(FName StatusBarName) const;
 	TSharedRef<SWidget> OnGetContentBrowser();
diff --git a/Engine/Source/Editor/UnrealEd/Classes/Analytics/AnalyticsPrivacySettings.h b/Engine/Source/Editor/UnrealEd/Classes/Analytics/AnalyticsPrivacySettings.h
--- a/Engine/Source/Editor/UnrealEd/Classes/Analytics/AnalyticsPrivacySettings.h
+++ b/Engine/Source/Editor/UnrealEd/Classes/Analytics/AnalyticsPrivacySettings.h
@@ -24,6 +24,19 @@ class UAnalyticsPrivacySettings : public UObject, public IImportantToggleSetting
 	UPROPERTY(EditAnywhere, config, Category=Options)
 	bool bSendUsageData;
 
+	/** Determines whether the editor asks for your industry and sends that information to Epic Games in order to improve Unreal Engine. Your information will never be shared with 3rd parties. 
+	* As Unreal Engine is used across a range of industries, we want to make sure we are addressing the needs of those verticals. A one-time prompt will show up for a user asking if they 
+	* are using the editor for the field of: Game Development, Film & TV, Architecture, Automotive & Manufacturing, Broadcasting & Live Events, Advertising & Marketing, Simulation & Training, or Other.
+	* 
+	* If you would like to block the prompt from showing up for members of your team, but otherwise send usage data back to Epic Games, you can add the following to your BaseEditorSettings.ini:
+	* [/Script/UnrealEd.AnalyticsPrivacySettings]
+	* bSuppressIndustryPopup=True
+	* 
+	* Additionally, for users who do not wish to send this information back, simply close the pop-up and it will not be shown again.
+	*/
+	UPROPERTY(EditAnywhere, config, Category=Options)
+	bool bSuppressIndustryPopup;
+
 public:
 
 	// BEGIN IImportantToggleSettingInterface
