Updating name of Training and Simulation industry

diff --git a/Engine/Source/Editor/StatusBar/Private/SOneTimeIndustryQuery.cpp b/Engine/Source/Editor/StatusBar/Private/SOneTimeIndustryQuery.cpp
--- a/Engine/Source/Editor/StatusBar/Private/SOneTimeIndustryQuery.cpp
+++ b/Engine/Source/Editor/StatusBar/Private/SOneTimeIndustryQuery.cpp
@@ -26,7 +26,7 @@ const FName SOneTimeIndustryQuery::ArchitectureIndustryName = FName("Architectur
 const FName SOneTimeIndustryQuery::AutoIndustryName = FName("AutomotiveAndManufacturing");
 const FName SOneTimeIndustryQuery::BroadcastIndustryName = FName("BroadcastingAndLiveEvents");
 const FName SOneTimeIndustryQuery::AdIndustryName = FName("AdvertisingAndMarketing");
-const FName SOneTimeIndustryQuery::SimulationIndustryName = FName("SimulationAndTraining");
+const FName SOneTimeIndustryQuery::SimulationIndustryName = FName("TrainingAndSimulation");
 const FName SOneTimeIndustryQuery::OtherIndustryName = FName("Other");
 
 void SOneTimeIndustryQuery::Show(TSharedPtr<SWindow> InParentWindow)
@@ -212,7 +212,7 @@ FText SOneTimeIndustryQuery::GetIndustryNameText(FName IndustryName) const
 	}
 	else if (IndustryName == SOneTimeIndustryQuery::SimulationIndustryName)
 	{
-		return LOCTEXT("SimulationAndTraining", "Simulation & Training");
+		return LOCTEXT("TrainingAndSimulation", "Training & Simulation");
 	}
 	else if (IndustryName == SOneTimeIndustryQuery::OtherIndustryName)
 	{
