Made 'Submit Changes' button invisible for Skein

diff --git a/Engine/Source/Editor/StatusBar/Private/SourceControlMenuHelpers.cpp b/Engine/Source/Editor/StatusBar/Private/SourceControlMenuHelpers.cpp
--- a/Engine/Source/Editor/StatusBar/Private/SourceControlMenuHelpers.cpp
+++ b/Engine/Source/Editor/StatusBar/Private/SourceControlMenuHelpers.cpp
@@ -66,7 +66,9 @@ void FSourceControlCommands::RegisterCommands()
 	ActionList->MapAction(
 		SubmitContent,
 		FExecuteAction::CreateLambda([]() { FSourceControlWindows::ChoosePackagesToCheckIn(); }),
-		FCanExecuteAction::CreateStatic(&FSourceControlWindows::CanChoosePackagesToCheckIn)
+		FCanExecuteAction::CreateStatic(&FSourceControlWindows::CanChoosePackagesToCheckIn),
+		FIsActionChecked::CreateLambda([]() { return false; }),
+		FIsActionButtonVisible::CreateStatic(&FSourceControlCommands::SubmitContent_IsVisible)
 	);
 
 	ActionList->MapAction(
@@ -93,6 +95,15 @@ bool FSourceControlCommands::ViewChangelists_IsVisible()
 	return ISourceControlModule::Get().GetProvider().UsesChangelists();
 }
 
+bool FSourceControlCommands::SubmitContent_IsVisible()
+{
+	if (FSourceControlMenuHelpers::GetSourceControlCheckInStatusVisibility() == EVisibility::Visible)
+	{
+		return false;
+	}
+	return true;
+}
+
 void FSourceControlCommands::ViewChangelists_Clicked()
 {
 	ISourceControlWindowsModule::Get().ShowChangelistsTab();
diff --git a/Engine/Source/Editor/StatusBar/Private/SourceControlMenuHelpers.h b/Engine/Source/Editor/StatusBar/Private/SourceControlMenuHelpers.h
--- a/Engine/Source/Editor/StatusBar/Private/SourceControlMenuHelpers.h
+++ b/Engine/Source/Editor/StatusBar/Private/SourceControlMenuHelpers.h
@@ -24,6 +24,7 @@ private:
 	static void ConnectToSourceControl_Clicked();
 	static bool ViewChangelists_CanExecute();
 	static bool ViewChangelists_IsVisible();
+	static bool SubmitContent_IsVisible();
 	static void ViewChangelists_Clicked();
 	static bool CheckOutModifiedFiles_CanExecute();
 	static void CheckOutModifiedFiles_Clicked();
@@ -43,6 +44,7 @@ public:
 
 class FSourceControlMenuHelpers
 {
+	friend class FSourceControlCommands;
 public:
 	/**
 	 * Static: Access singleton instance
