Adding Fashion to the industry analytics pop-up

diff --git a/Engine/Source/Editor/StatusBar/Private/SOneTimeIndustryQuery.cpp b/Engine/Source/Editor/StatusBar/Private/SOneTimeIndustryQuery.cpp
--- a/Engine/Source/Editor/StatusBar/Private/SOneTimeIndustryQuery.cpp
+++ b/Engine/Source/Editor/StatusBar/Private/SOneTimeIndustryQuery.cpp
@@ -27,6 +27,7 @@ const FName SOneTimeIndustryQuery::AutoIndustryName = FName("AutomotiveAndManufa
 const FName SOneTimeIndustryQuery::BroadcastIndustryName = FName("BroadcastingAndLiveEvents");
 const FName SOneTimeIndustryQuery::AdIndustryName = FName("AdvertisingAndMarketing");
 const FName SOneTimeIndustryQuery::SimulationIndustryName = FName("TrainingAndSimulation");
+const FName SOneTimeIndustryQuery::FashionIndustryName = FName("Fashion");
 const FName SOneTimeIndustryQuery::OtherIndustryName = FName("Other");
 
 void SOneTimeIndustryQuery::Show(TSharedPtr<SWindow> InParentWindow)
@@ -77,6 +78,7 @@ void SOneTimeIndustryQuery::Construct(const FArguments& InArgs)
 	IndustryComboList.Add(MakeShared<FName>(SOneTimeIndustryQuery::BroadcastIndustryName));
 	IndustryComboList.Add(MakeShared<FName>(SOneTimeIndustryQuery::AdIndustryName));
 	IndustryComboList.Add(MakeShared<FName>(SOneTimeIndustryQuery::SimulationIndustryName));
+	IndustryComboList.Add(MakeShared<FName>(SOneTimeIndustryQuery::FashionIndustryName));
 	IndustryComboList.Add(MakeShared<FName>(SOneTimeIndustryQuery::OtherIndustryName));
 
 	ChildSlot
@@ -214,6 +216,10 @@ FText SOneTimeIndustryQuery::GetIndustryNameText(FName IndustryName) const
 	{
 		return LOCTEXT("TrainingAndSimulation", "Training & Simulation");
 	}
+	else if (IndustryName == SOneTimeIndustryQuery::FashionIndustryName)
+	{
+		return LOCTEXT("Fashion", "Fashion");
+	}
 	else if (IndustryName == SOneTimeIndustryQuery::OtherIndustryName)
 	{
 		return LOCTEXT("Other", "Other");
diff --git a/Engine/Source/Editor/StatusBar/Private/SOneTimeIndustryQuery.h b/Engine/Source/Editor/StatusBar/Private/SOneTimeIndustryQuery.h
--- a/Engine/Source/Editor/StatusBar/Private/SOneTimeIndustryQuery.h
+++ b/Engine/Source/Editor/StatusBar/Private/SOneTimeIndustryQuery.h
@@ -50,6 +50,7 @@ private:
 	static const FName BroadcastIndustryName;
 	static const FName AdIndustryName;
 	static const FName SimulationIndustryName;
+	static const FName FashionIndustryName;
 	static const FName OtherIndustryName;
 };
 
