Renamed 'Get Latest' status bar button to 'Sync Latest'

diff --git a/Engine/Source/Editor/StatusBar/Private/SourceControlMenuHelpers.cpp b/Engine/Source/Editor/StatusBar/Private/SourceControlMenuHelpers.cpp
--- a/Engine/Source/Editor/StatusBar/Private/SourceControlMenuHelpers.cpp
+++ b/Engine/Source/Editor/StatusBar/Private/SourceControlMenuHelpers.cpp
@@ -320,7 +320,7 @@ FText FSourceControlMenuHelpers::GetSourceControlSyncStatusText()
 	{
 		return LOCTEXT("SyncLatestButtonAtHeadText", "At Latest");
 	}
-	return LOCTEXT("SyncLatestButtonNotAtHeadText", "Get Latest");
+	return LOCTEXT("SyncLatestButtonNotAtHeadText", "Sync Latest");
 }
 
 FText FSourceControlMenuHelpers::GetSourceControlSyncStatusTooltipText()
