Fixed: Cannot get latest when there is a conflicting change

diff --git a/Engine/Source/Editor/StatusBar/Private/SourceControlMenuHelpers.cpp b/Engine/Source/Editor/StatusBar/Private/SourceControlMenuHelpers.cpp
--- a/Engine/Source/Editor/StatusBar/Private/SourceControlMenuHelpers.cpp
+++ b/Engine/Source/Editor/StatusBar/Private/SourceControlMenuHelpers.cpp
@@ -21,6 +21,7 @@
 #include "PackageTools.h"
 #include "AssetViewUtils.h"
 #include "Misc/ConfigCacheIni.h"
+#include "Misc/MessageDialog.h"
 
 #define LOCTEXT_NAMESPACE "SourceControlCommands"
 
@@ -355,15 +356,20 @@ FReply FSourceControlMenuHelpers::OnSourceControlSyncClicked()
 void FSourceControlMenuHelpers::SyncProject()
 {
 	const bool bSaved = SaveDirtyPackages();
-	if (bSaved)
+	EAppReturnType::Type DialogResult = EAppReturnType::No;
+
+	if (!bSaved)
 	{
-		AssetViewUtils::SyncPackagesFromSourceControl(ListAllPackages());
+		FText Title = LOCTEXT("UnsavedWarningTitle", "Unsaved changes");
+
+		DialogResult = FMessageDialog::Open(EAppMsgType::YesNo,
+			LOCTEXT("UnsavedWarningText", "Warning: There are modified assets which are not being saved. If you sync to latest you may lose your unsaved changes. Do you want to continue?"),
+			&Title);
 	}
-	else
+
+	if (bSaved || DialogResult == EAppReturnType::Yes)
 	{
-		FMessageLog SourceControlLog("SourceControl");
-		SourceControlLog.Warning(LOCTEXT("SourceControlMenu_Sync_Unsaved", "Save All Assets before attempting to Sync!"));
-		SourceControlLog.Notify();
+		AssetViewUtils::SyncPackagesFromSourceControl(ListAllPackages());
 	}
 }
 
