Swapped positions of Sync Latest and Check-in Changes Status Bar buttons

diff --git a/Engine/Source/Editor/StatusBar/Private/SourceControlMenuHelpers.cpp b/Engine/Source/Editor/StatusBar/Private/SourceControlMenuHelpers.cpp
--- a/Engine/Source/Editor/StatusBar/Private/SourceControlMenuHelpers.cpp
+++ b/Engine/Source/Editor/StatusBar/Private/SourceControlMenuHelpers.cpp
@@ -331,7 +331,7 @@ FText FSourceControlMenuHelpers::GetSourceControlSyncStatusTooltipText()
 	{
 		return LOCTEXT("SyncLatestButtonAtHeadTooltipText", "Currently at the latest Snapshot for this project");
 	}
-	return LOCTEXT("SyncLatestButtonNotAtHeadTooltipText", "Get the latest Snapshot for this project");
+	return LOCTEXT("SyncLatestButtonNotAtHeadTooltipText", "Sync to the latest Snapshot for this project");
 }
 
 const FSlateBrush* FSourceControlMenuHelpers::GetSourceControlSyncStatusIcon()
@@ -507,15 +507,16 @@ TSharedRef<SWidget> FSourceControlMenuHelpers::MakeSourceControlStatusWidget()
 			//       ensure its position with respect to the source control button.
 			FUnsavedAssetsTrackerModule::Get().MakeUnsavedAssetsStatusBarWidget()
 		]
-		+ SHorizontalBox::Slot()
+		+ SHorizontalBox::Slot() // Check In Changes Button
 		.VAlign(VAlign_Center)
+		.Padding(FMargin(8.0f, 0.0f, 4.0f, 0.0f))
 		.AutoWidth()
 		[
 			SNew(SButton)
 			.ButtonStyle(&FAppStyle::Get().GetWidgetStyle<FButtonStyle>("StatusBar.StatusBarButton"))
-			.ToolTipText_Static(&FSourceControlMenuHelpers::GetSourceControlSyncStatusTooltipText)
-			.Visibility_Static(&FSourceControlMenuHelpers::GetSourceControlSyncStatusVisibility)
-			.IsEnabled_Lambda([]() { return !IsAtLatestRevision(); })
+			.ToolTipText_Static(&FSourceControlMenuHelpers::GetSourceControlCheckInStatusTooltipText)
+			.Visibility_Static(&FSourceControlMenuHelpers::GetSourceControlCheckInStatusVisibility)
+			.IsEnabled_Lambda([]() { return GetNumLocalChanges() > 0; })
 			[
 				SNew(SHorizontalBox)
 				+ SHorizontalBox::Slot()
@@ -524,7 +525,7 @@ TSharedRef<SWidget> FSourceControlMenuHelpers::MakeSourceControlStatusWidget()
 				.HAlign(HAlign_Center)
 				[
 					SNew(SImage)
-					.Image_Static(&FSourceControlMenuHelpers::GetSourceControlSyncStatusIcon)
+					.Image_Static(&FSourceControlMenuHelpers::GetSourceControlCheckInStatusIcon)
 				]
 				+ SHorizontalBox::Slot()
 				.AutoWidth()
@@ -533,21 +534,20 @@ TSharedRef<SWidget> FSourceControlMenuHelpers::MakeSourceControlStatusWidget()
 				[
 					SNew(STextBlock)
 					.TextStyle(&FAppStyle::Get().GetWidgetStyle<FTextBlockStyle>("NormalText"))
-					.Text_Static(&FSourceControlMenuHelpers::GetSourceControlSyncStatusText)
+					.Text_Static(&FSourceControlMenuHelpers::GetSourceControlCheckInStatusText)
 				]
 			]
-			.OnClicked_Static(&FSourceControlMenuHelpers::OnSourceControlSyncClicked)
+			.OnClicked_Static(&FSourceControlMenuHelpers::OnSourceControlCheckInChangesClicked)
 		]
-		+ SHorizontalBox::Slot()
+		+ SHorizontalBox::Slot() // Sync Latest Button
 		.VAlign(VAlign_Center)
-		.Padding(FMargin(8.0f, 0.0f, 4.0f, 0.0f))
 		.AutoWidth()
 		[
 			SNew(SButton)
 			.ButtonStyle(&FAppStyle::Get().GetWidgetStyle<FButtonStyle>("StatusBar.StatusBarButton"))
-			.ToolTipText_Static(&FSourceControlMenuHelpers::GetSourceControlCheckInStatusTooltipText)
-			.Visibility_Static(&FSourceControlMenuHelpers::GetSourceControlCheckInStatusVisibility)
-			.IsEnabled_Lambda([]() { return GetNumLocalChanges() > 0; })
+			.ToolTipText_Static(&FSourceControlMenuHelpers::GetSourceControlSyncStatusTooltipText)
+			.Visibility_Static(&FSourceControlMenuHelpers::GetSourceControlSyncStatusVisibility)
+			.IsEnabled_Lambda([]() { return !IsAtLatestRevision(); })
 			[
 				SNew(SHorizontalBox)
 				+ SHorizontalBox::Slot()
@@ -556,7 +556,7 @@ TSharedRef<SWidget> FSourceControlMenuHelpers::MakeSourceControlStatusWidget()
 				.HAlign(HAlign_Center)
 				[
 					SNew(SImage)
-					.Image_Static(&FSourceControlMenuHelpers::GetSourceControlCheckInStatusIcon)
+					.Image_Static(&FSourceControlMenuHelpers::GetSourceControlSyncStatusIcon)
 				]
 				+ SHorizontalBox::Slot()
 				.AutoWidth()
@@ -565,12 +565,12 @@ TSharedRef<SWidget> FSourceControlMenuHelpers::MakeSourceControlStatusWidget()
 				[
 					SNew(STextBlock)
 					.TextStyle(&FAppStyle::Get().GetWidgetStyle<FTextBlockStyle>("NormalText"))
-					.Text_Static(&FSourceControlMenuHelpers::GetSourceControlCheckInStatusText)
+					.Text_Static(&FSourceControlMenuHelpers::GetSourceControlSyncStatusText)
 				]
 			]
-			.OnClicked_Static(&FSourceControlMenuHelpers::OnSourceControlCheckInChangesClicked)
+			.OnClicked_Static(&FSourceControlMenuHelpers::OnSourceControlSyncClicked)
 		]
-		+ SHorizontalBox::Slot()
+		+ SHorizontalBox::Slot() // Source Control Menu
 		.VAlign(VAlign_Center)
 		.AutoWidth()
 		[
