SourceControl: Fixed clicking 'X' or 'Esc' on the 'Revert All' confirmation modal not resulting in the cancellation of the operation.

diff --git a/Engine/Source/Editor/StatusBar/Private/SourceControlMenuHelpers.cpp b/Engine/Source/Editor/StatusBar/Private/SourceControlMenuHelpers.cpp
--- a/Engine/Source/Editor/StatusBar/Private/SourceControlMenuHelpers.cpp
+++ b/Engine/Source/Editor/StatusBar/Private/SourceControlMenuHelpers.cpp
@@ -158,12 +158,10 @@ void FSourceControlCommands::RevertAllModifiedFiles_Clicked()
 {
 	FText Message = LOCTEXT("RevertAllModifiedFiles", "Are you sure you want to revert all local changes? By proceeding, your local changes will be discarded, and the state of your last synced snapshot will be restored.");
 	FText Title = LOCTEXT("RevertAllModifiedFiles_Title", "Revert all local changes");
-	if (FMessageDialog::Open(EAppMsgType::YesNo, EAppReturnType::No, Message, &Title) == EAppReturnType::No)
+	if (FMessageDialog::Open(EAppMsgType::YesNo, EAppReturnType::No, Message, &Title) == EAppReturnType::Yes)
 	{
-		return;
+		FSourceControlWindows::RevertAllChangesAndReloadAllPackages();
 	}
-	
-	FSourceControlWindows::RevertAllChangesAndReloadAllPackages();
 }
 
 FSourceControlMenuHelpers& FSourceControlMenuHelpers::Get()
