View Changes UI visual improvements

- Side panel no longer compressed
- Name changed to "View Changes"

diff --git a/Engine/Source/Editor/MainFrame/Private/Frame/MainFrameActions.cpp b/Engine/Source/Editor/MainFrame/Private/Frame/MainFrameActions.cpp
--- a/Engine/Source/Editor/MainFrame/Private/Frame/MainFrameActions.cpp
+++ b/Engine/Source/Editor/MainFrame/Private/Frame/MainFrameActions.cpp
@@ -105,7 +105,7 @@ void FMainFrameCommands::RegisterCommands()
 	UI_COMMAND( ChooseFilesToSave, "Choose Files to Save...", "Opens a dialog with save options for content and levels", EUserInterfaceActionType::Button, FInputChord( EModifierKey::Control | EModifierKey::Alt | EModifierKey::Shift, EKeys::S ) );
 	ActionList->MapAction( ChooseFilesToSave, FExecuteAction::CreateStatic( &FMainFrameActionCallbacks::ChoosePackagesToSave ), FCanExecuteAction::CreateStatic( &FMainFrameActionCallbacks::CanSaveWorld ) );
 
-	UI_COMMAND( ViewChangelists, "View Changelists", "Opens a dialog displaying current changelists.", EUserInterfaceActionType::Button, FInputChord() );
+	UI_COMMAND( ViewChangelists, "View Changes", "Opens a dialog displaying current changes.", EUserInterfaceActionType::Button, FInputChord() );
 	ActionList->MapAction(ViewChangelists, FExecuteAction::CreateStatic( &FMainFrameActionCallbacks::ViewChangelists ), FCanExecuteAction::CreateStatic( &FMainFrameActionCallbacks::CanViewChangelists ) );
 
 	UI_COMMAND( SubmitContent, "Submit Content", "Opens a dialog with check in options for content and levels.", EUserInterfaceActionType::Button, FInputChord() );
diff --git a/Engine/Source/Editor/SourceControlWindows/Private/SSourceControlChangelists.cpp b/Engine/Source/Editor/SourceControlWindows/Private/SSourceControlChangelists.cpp
--- a/Engine/Source/Editor/SourceControlWindows/Private/SSourceControlChangelists.cpp
+++ b/Engine/Source/Editor/SourceControlWindows/Private/SSourceControlChangelists.cpp
@@ -465,7 +465,6 @@ void SSourceControlChangelistsWidget::Construct(const FArguments& InArgs)
 					]
 				]
 				+SVerticalBox::Slot()
-				.AutoHeight()
 				[
 					SNew(SOverlay)
 					+SOverlay::Slot()
@@ -492,7 +491,7 @@ void SSourceControlChangelistsWidget::Construct(const FArguments& InArgs)
 						[
 							ChangelistExpandableArea.ToSharedRef()
 						]
-	
+						
 						// Bottom slot: Uncontrolled Changelists
 						+SSplitter::Slot()
 						.SizeRule_Lambda([this](){ return UncontrolledChangelistExpandableArea->IsExpanded() ? SSplitter::ESizeRule::FractionOfParent : SSplitter::ESizeRule::SizeToContent; })
diff --git a/Engine/Source/Editor/SourceControlWindows/Private/SourceControlWindowsModule.cpp b/Engine/Source/Editor/SourceControlWindows/Private/SourceControlWindowsModule.cpp
--- a/Engine/Source/Editor/SourceControlWindows/Private/SourceControlWindowsModule.cpp
+++ b/Engine/Source/Editor/SourceControlWindows/Private/SourceControlWindowsModule.cpp
@@ -66,8 +66,8 @@ void FSourceControlWindowsModule::StartupModule()
 
 	// Register the changelist tab spawner
 	FGlobalTabmanager::Get()->RegisterNomadTabSpawner(SourceControlChangelistsTabName, FOnSpawnTab::CreateRaw(this, &FSourceControlWindowsModule::CreateChangelistsTab))
-		.SetDisplayName(LOCTEXT("ChangelistsTabTitle", "View Changelists"))
-		.SetTooltipText(LOCTEXT("ChangelistsTabTooltip", "Opens a dialog displaying current changelists."))
+		.SetDisplayName(LOCTEXT("ChangelistsTabTitle", "View Changes"))
+		.SetTooltipText(LOCTEXT("ChangelistsTabTooltip", "Opens a dialog displaying current changes."))
 		.SetIcon(SourceControlIcon)
 		.SetMenuType(ETabSpawnerMenuType::Hidden);
 
diff --git a/Engine/Source/Editor/StatusBar/Private/SourceControlMenuHelpers.cpp b/Engine/Source/Editor/StatusBar/Private/SourceControlMenuHelpers.cpp
--- a/Engine/Source/Editor/StatusBar/Private/SourceControlMenuHelpers.cpp
+++ b/Engine/Source/Editor/StatusBar/Private/SourceControlMenuHelpers.cpp
@@ -47,7 +47,7 @@ void FSourceControlCommands::RegisterCommands()
 {
 	UI_COMMAND(ConnectToSourceControl, "Connect to Revision Control...", "Connect to a revision control system for tracking changes to your content and levels.", EUserInterfaceActionType::Button, FInputChord());
 	UI_COMMAND(ChangeSourceControlSettings, "Change Revision Control Settings...", "Opens a dialog to change revision control settings.", EUserInterfaceActionType::Button, FInputChord());
-	UI_COMMAND(ViewChangelists, "View Changelists", "Opens a dialog displaying current changelists.", EUserInterfaceActionType::Button, FInputChord());
+	UI_COMMAND(ViewChangelists, "View Changes", "Opens a dialog displaying current changes.", EUserInterfaceActionType::Button, FInputChord());
 	UI_COMMAND(SubmitContent, "Submit Content", "Opens a dialog with check in options for content and levels.", EUserInterfaceActionType::Button, FInputChord());
 	UI_COMMAND(CheckOutModifiedFiles, "Check Out Modified Files", "Opens a dialog to check out any assets which have been modified.", EUserInterfaceActionType::Button, FInputChord());
 	UI_COMMAND(RevertAll, "Revert All Files", "Opens a dialog to revert any assets which have been modified.", EUserInterfaceActionType::Button, FInputChord());
