SNotificationList: Improve overflow behavior of text by adding ellipses and hover text

diff --git a/Engine/Source/Editor/StatusBar/Private/NotificationTesting.cpp b/Engine/Source/Editor/StatusBar/Private/NotificationTesting.cpp
--- a/Engine/Source/Editor/StatusBar/Private/NotificationTesting.cpp
+++ b/Engine/Source/Editor/StatusBar/Private/NotificationTesting.cpp
@@ -205,6 +205,38 @@ static void TestNotifications()
 			NotificationInfo.ButtonDetails.Add(Button2);
 
 
+			auto Notification = FSlateNotificationManager::Get().AddNotification(NotificationInfo);
+			Notification->SetCompletionState(SNotificationItem::CS_Pending);
+
+			Notification->ExpireAndFadeout();
+
+			return false;
+		});
+
+	FTSTicker::GetCoreTicker().AddTicker(TEXT("TestNotifications"), 9.0f, [](float DeltaTime)
+		{
+			FNotificationInfo NotificationInfo(FText::FromString(TEXT("This one has a lot of text on the buttons which should overflow properly")));
+			NotificationInfo.FadeInDuration = 2.0f;
+			NotificationInfo.FadeOutDuration = 2.0f;
+			NotificationInfo.ExpireDuration = Timeout;
+
+			NotificationInfo.CheckBoxText = FText::FromString(TEXT("This is a checkbox with a lot of text. Hover over it to read the full text in the tooltip."));
+			NotificationInfo.CheckBoxState = ECheckBoxState::Checked;
+			NotificationInfo.CheckBoxStateChanged = FOnCheckStateChanged::CreateStatic([](ECheckBoxState NewState) {});
+
+			NotificationInfo.Hyperlink = FSimpleDelegate::CreateLambda([]() {});
+			NotificationInfo.HyperlinkText = FText::FromString(TEXT("This is a hyperlink with a lot of text. Hover over it to read the full text in the tooltip."));
+
+			NotificationInfo.bUseSuccessFailIcons = true;
+			NotificationInfo.bUseThrobber = true;
+
+
+			FNotificationButtonInfo Button1(FText::FromString("This is a button with a lot of text. Hover over it to read the full text in the tooltip."), FText::GetEmpty(), FSimpleDelegate());
+			FNotificationButtonInfo Button2(FText::FromString("This is another button with a lot of text. Hover over it to read the full text in the tooltip."), FText::GetEmpty(), FSimpleDelegate());
+
+			NotificationInfo.ButtonDetails.Add(Button1);
+			NotificationInfo.ButtonDetails.Add(Button2);
+
 			auto Notification = FSlateNotificationManager::Get().AddNotification(NotificationInfo);
 			Notification->SetCompletionState(SNotificationItem::CS_Pending);
 
diff --git a/Engine/Source/Runtime/Slate/Private/Widgets/Notifications/SNotificationList.cpp b/Engine/Source/Runtime/Slate/Private/Widgets/Notifications/SNotificationList.cpp
--- a/Engine/Source/Runtime/Slate/Private/Widgets/Notifications/SNotificationList.cpp
+++ b/Engine/Source/Runtime/Slate/Private/Widgets/Notifications/SNotificationList.cpp
@@ -462,6 +462,8 @@ public:
 				.Visibility(this, &SNotificationItemImpl::GetHyperlinkVisibility)
 				.Text(this, &SNotificationItemImpl::GetHyperlinkText)
 				.OnNavigate(this, &SNotificationItemImpl::OnHyperlinkClicked)
+				.ToolTipText(this, &SNotificationItemImpl::GetHyperlinkText)
+				.TextStyle(&FAppStyle::GetWidgetStyle<FTextBlockStyle>("NotificationList.WidgetText"))
 			];
 		}
 
@@ -480,6 +482,8 @@ public:
 				[
 					SNew(STextBlock)
 					.Text(InArgs._CheckBoxText)
+					.TextStyle(&FAppStyle::GetWidgetStyle<FTextBlockStyle>("NotificationList.WidgetText"))
+					.ToolTipText(InArgs._CheckBoxText)
 				]
 			];
 		}
@@ -494,13 +498,13 @@ public:
 
 				ButtonsBox->AddSlot()
 				.AutoWidth()
-
 				.VAlign(VAlign_Center)
 				.Padding(0.0f, 0.0f, 4.0f, 0.0f)
 				[
 					SNew(SButton)
 					.Text(Button.Text)
-					.ToolTipText(Button.ToolTip)
+					.ToolTipText(Button.ToolTip.IsEmpty() ? Button.Text : Button.ToolTip)
+					.TextStyle(&FAppStyle::GetWidgetStyle<FTextBlockStyle>("NotificationList.WidgetText"))
 					.OnClicked(this, &SNotificationItemImpl::OnButtonClicked, Button.Callback) 
 					.Visibility( this, &SNotificationItemImpl::GetButtonVisibility, Button.VisibilityOnNone, Button.VisibilityOnPending, Button.VisibilityOnSuccess, Button.VisibilityOnFail )
 				];
diff --git a/Engine/Source/Runtime/SlateCore/Private/Styling/StarshipCoreStyle.cpp b/Engine/Source/Runtime/SlateCore/Private/Styling/StarshipCoreStyle.cpp
--- a/Engine/Source/Runtime/SlateCore/Private/Styling/StarshipCoreStyle.cpp
+++ b/Engine/Source/Runtime/SlateCore/Private/Styling/StarshipCoreStyle.cpp
@@ -669,7 +669,9 @@ TSharedRef<ISlateStyle> FStarshipCoreStyle::Create()
 		Style->Set("NotificationList.SuccessImage", new IMAGE_BRUSH("Icons/notificationlist_success", Icon16x16));
 		Style->Set("NotificationList.FailImage", new IMAGE_BRUSH("Icons/notificationlist_fail", Icon16x16));
 
-	
+		Style->Set("NotificationList.WidgetText", FTextBlockStyle(NormalText)
+		.SetOverflowPolicy(ETextOverflowPolicy::Ellipsis));
+
 	}
 
 	// SSeparator defaults...
