Implemented auto-save on check-in.
Fixed issue that prevented deleted objects from being restored on Revert All
As a collateral, fixed issue that caused not saving deleted files during Delete workflow to result in errors after sync
As a collateral, fixed issue that caused added and saved files to disappear after Sync Latest

diff --git a/Engine/Source/Developer/SourceControl/Private/SourceControlHelpers.cpp b/Engine/Source/Developer/SourceControl/Private/SourceControlHelpers.cpp
--- a/Engine/Source/Developer/SourceControl/Private/SourceControlHelpers.cpp
+++ b/Engine/Source/Developer/SourceControl/Private/SourceControlHelpers.cpp
@@ -1132,6 +1132,43 @@ bool USourceControlHelpers::ListRevertablePackages(TArray<FString>& OutRevertabl
 	TMap<FString, FSourceControlStatePtr> PackageStates;
 	FEditorFileUtils::FindAllSubmittablePackageFiles(PackageStates, true);
 
+	TArray<FString> PackageNames;
+	PackageStates.GetKeys(PackageNames);
+	TMap<FString, FString> FileNamesToPackageNames;
+
+	// Get a list of files pending delete
+	TArray<FSourceControlStateRef> PendingDeleteItems = SourceControlProvider.GetCachedStateByPredicate(
+		[&PackageNames, &FileNamesToPackageNames](const FSourceControlStateRef& State)
+		{
+			const FString& Filename = State->GetFilename();
+
+			FString PackageName;
+			FString FailureReason;
+			if (!FPackageName::TryConvertFilenameToLongPackageName(Filename, PackageName, &FailureReason))
+			{
+				UE_LOG(LogSourceControl, Warning, TEXT("%s"), *FailureReason);
+				return false;
+			}
+
+			if (State->IsDeleted() && !PackageNames.Contains(PackageName))
+			{
+				FileNamesToPackageNames.Add(Filename, PackageName);
+				return true;
+			}
+			return false;
+		}
+	);
+
+	// And append them to the list
+	for (FSourceControlStateRef& Item : PendingDeleteItems)
+	{
+		const FString* PackageName = FileNamesToPackageNames.Find(Item->GetFilename());
+		if (PackageName)
+		{
+			PackageStates.Add(*PackageName, Item);
+		}
+	}
+
 	for (auto& PackageState : PackageStates)
 	{
 		const FString PackageName = PackageState.Key;
diff --git a/Engine/Source/Developer/UnsavedAssetsTracker/Source/Private/UnsavedAssetsTracker.cpp b/Engine/Source/Developer/UnsavedAssetsTracker/Source/Private/UnsavedAssetsTracker.cpp
--- a/Engine/Source/Developer/UnsavedAssetsTracker/Source/Private/UnsavedAssetsTracker.cpp
+++ b/Engine/Source/Developer/UnsavedAssetsTracker/Source/Private/UnsavedAssetsTracker.cpp
@@ -507,7 +507,7 @@ void FUnsavedAssetsTracker::OnSourceControlFileStatusUpdate(const FString& AbsFi
 	}
 }
 
-void FUnsavedAssetsTracker::PrompToSavePackages()
+bool FUnsavedAssetsTracker::PromptToSavePackages()
 {
 	if (GetUnsavedAssetNum() > 0)
 	{
@@ -526,8 +526,10 @@ void FUnsavedAssetsTracker::PrompToSavePackages()
 
 			// Stay in sync with what the package the engine thinks is dirty.
 			SyncWithDirtyPackageList();
+			return true;
 		}
 	}
+	return false;
 }
 
 bool FUnsavedAssetsTracker::IsAssetUnsaved(const FString& FileAbsPathname) const
diff --git a/Engine/Source/Developer/UnsavedAssetsTracker/Source/Private/UnsavedAssetsTracker.h b/Engine/Source/Developer/UnsavedAssetsTracker/Source/Private/UnsavedAssetsTracker.h
--- a/Engine/Source/Developer/UnsavedAssetsTracker/Source/Private/UnsavedAssetsTracker.h
+++ b/Engine/Source/Developer/UnsavedAssetsTracker/Source/Private/UnsavedAssetsTracker.h
@@ -41,7 +41,7 @@ public:
 	void SetWarningNotificationEnabled(bool bEnabled) { bWarningNotificationEnabled = bEnabled; }
 
 	/** Displays a dialog prompting the user to save the packages. */
-	void PrompToSavePackages();
+	bool PromptToSavePackages();
 
 	/** Check if the input asset is unsaved. */
 	bool IsAssetUnsaved(const FString& FileAbsPathname) const;
diff --git a/Engine/Source/Developer/UnsavedAssetsTracker/Source/Private/UnsavedAssetsTrackerModule.cpp b/Engine/Source/Developer/UnsavedAssetsTracker/Source/Private/UnsavedAssetsTrackerModule.cpp
--- a/Engine/Source/Developer/UnsavedAssetsTracker/Source/Private/UnsavedAssetsTrackerModule.cpp
+++ b/Engine/Source/Developer/UnsavedAssetsTracker/Source/Private/UnsavedAssetsTrackerModule.cpp
@@ -64,11 +64,21 @@ TSharedRef<SWidget> FUnsavedAssetsTrackerModule::MakeUnsavedAssetsStatusBarWidge
 		{
 			if (TSharedPtr<FUnsavedAssetsTracker> Tracker = WeakUnsavedAssetTracker.Pin())
 			{
-				Tracker->PrompToSavePackages();
+				Tracker->PromptToSavePackages();
 			}
 			return FReply::Handled();
 		});
 }
 
+bool FUnsavedAssetsTrackerModule::PromptToSavePackages()
+{
+	if (!UnsavedAssetTracker)
+	{
+		return false;
+	}
+
+	return UnsavedAssetTracker->PromptToSavePackages();
+}
+
 IMPLEMENT_MODULE(FUnsavedAssetsTrackerModule, UnsavedAssetsTracker)
 
diff --git a/Engine/Source/Developer/UnsavedAssetsTracker/Source/Public/UnsavedAssetsTrackerModule.h b/Engine/Source/Developer/UnsavedAssetsTracker/Source/Public/UnsavedAssetsTrackerModule.h
--- a/Engine/Source/Developer/UnsavedAssetsTracker/Source/Public/UnsavedAssetsTrackerModule.h
+++ b/Engine/Source/Developer/UnsavedAssetsTracker/Source/Public/UnsavedAssetsTrackerModule.h
@@ -58,6 +58,9 @@ public:
 	/** Construct the widget for the Editor status bar.*/
 	TSharedRef<SWidget> MakeUnsavedAssetsStatusBarWidget();
 
+	/** Displays a dialog prompting the user to save unsaved packages. */
+	bool PromptToSavePackages();
+
 	/** Invoked when a file is added to the unsaved list. */
 	FOnUnsavedAssetAdded OnUnsavedAssetAdded;
 
diff --git a/Engine/Source/Editor/StatusBar/Private/SourceControlMenuHelpers.cpp b/Engine/Source/Editor/StatusBar/Private/SourceControlMenuHelpers.cpp
--- a/Engine/Source/Editor/StatusBar/Private/SourceControlMenuHelpers.cpp
+++ b/Engine/Source/Editor/StatusBar/Private/SourceControlMenuHelpers.cpp
@@ -161,6 +161,23 @@ void FSourceControlCommands::RevertAllModifiedFiles_Clicked()
 	FText Title = LOCTEXT("RevertAllModifiedFiles_Title", "Revert all local changes");
 	if (FMessageDialog::Open(EAppMsgType::YesNo, EAppReturnType::No, Message, &Title) == EAppReturnType::Yes)
 	{
+		// Get a list of all the unsaved packages
+		TArray<FString> UnsavedFileNames = FUnsavedAssetsTrackerModule::Get().GetUnsavedAssets();
+		TArray<UPackage*> UnsavedPackages;
+		UnsavedPackages.Reserve(UnsavedFileNames.Num());
+
+		for (FString& FileName : UnsavedFileNames)
+		{
+			FString PackageName = UPackageTools::FilenameToPackageName(FileName);
+			UPackage* Package = FindPackage(nullptr, *PackageName);
+			if (Package != nullptr)
+			{
+				UnsavedPackages.Add(Package);
+			}
+		}
+
+		UEditorLoadingAndSavingUtils::SavePackages(UnsavedPackages, /*bOnlyDirty=*/true);
+		
 		FBookmarkScoped BookmarkScoped;
 		FSourceControlWindows::RevertAllChangesAndReloadWorld();
 	}
@@ -404,7 +421,7 @@ FReply FSourceControlMenuHelpers::OnSourceControlSyncClicked()
 int FSourceControlMenuHelpers::GetNumLocalChanges()
 {
 	ISourceControlModule& SourceControlModule = ISourceControlModule::Get();
-	if (SourceControlModule.IsEnabled() && SourceControlModule.GetProvider().IsAvailable() 
+	if (SourceControlModule.IsEnabled() && SourceControlModule.GetProvider().IsAvailable()
 		&& SourceControlModule.GetProvider().GetNumLocalChanges().IsSet())
 	{
 		return SourceControlModule.GetProvider().GetNumLocalChanges().GetValue();
@@ -412,6 +429,11 @@ int FSourceControlMenuHelpers::GetNumLocalChanges()
 	return 0;
 }
 
+bool FSourceControlMenuHelpers::HasLocalChanges()
+{
+	return GetNumLocalChanges() > 0 || FUnsavedAssetsTrackerModule::Get().GetUnsavedAssetNum() > 0;
+}
+
 EVisibility FSourceControlMenuHelpers::GetSourceControlCheckInStatusVisibility()
 {
 	bool bDisplaySourceControlCheckInStatus = false;
@@ -427,7 +449,7 @@ EVisibility FSourceControlMenuHelpers::GetSourceControlCheckInStatusVisibility()
 
 FText FSourceControlMenuHelpers::GetSourceControlCheckInStatusText()
 {
-	if (GetNumLocalChanges() > 0)
+	if (HasLocalChanges())
 	{
 		return LOCTEXT("CheckInButtonChangesText", "Check-in Changes");
 	}
@@ -437,9 +459,10 @@ FText FSourceControlMenuHelpers::GetSourceControlCheckInStatusText()
 
 FText FSourceControlMenuHelpers::GetSourceControlCheckInStatusTooltipText()
 {
-	if (GetNumLocalChanges() > 0)
+	if (HasLocalChanges())
 	{
-		return FText::Format(LOCTEXT("CheckInButtonChangesTooltipText", "Check-in {0} change(s) to this project"), GetNumLocalChanges());
+		return FText::Format(LOCTEXT("CheckInButtonChangesTooltipText", "Check-in {0} change(s) to this project"), 
+			GetNumLocalChanges() > 0 ? GetNumLocalChanges() : FUnsavedAssetsTrackerModule::Get().GetUnsavedAssetNum());
 	}
 	return LOCTEXT("CheckInButtonNoChangesTooltipText", "No Changes to check in for this project");
 }
@@ -449,7 +472,7 @@ const FSlateBrush* FSourceControlMenuHelpers::GetSourceControlCheckInStatusIcon(
 	static const FSlateBrush* NoLocalChangesBrush = FRevisionControlStyleManager::Get().GetBrush("RevisionControl.StatusBar.NoLocalChanges");
 	static const FSlateBrush* HasLocalChangesBrush = FRevisionControlStyleManager::Get().GetBrush("RevisionControl.StatusBar.HasLocalChanges");
 
-	if (GetNumLocalChanges() > 0)
+	if (HasLocalChanges())
 	{
 		return HasLocalChangesBrush;
 	}
@@ -458,7 +481,7 @@ const FSlateBrush* FSourceControlMenuHelpers::GetSourceControlCheckInStatusIcon(
 
 FReply FSourceControlMenuHelpers::OnSourceControlCheckInChangesClicked()
 {
-	if (FSourceControlWindows::CanChoosePackagesToCheckIn())
+	if (FUnsavedAssetsTrackerModule::Get().PromptToSavePackages() || HasLocalChanges())
 	{
 		FSourceControlWindows::ChoosePackagesToCheckIn();
 	}
@@ -495,7 +518,7 @@ TSharedRef<SWidget> FSourceControlMenuHelpers::MakeSourceControlStatusWidget()
 			.ButtonStyle(&FAppStyle::Get().GetWidgetStyle<FButtonStyle>("StatusBar.StatusBarButton"))
 			.ToolTipText_Static(&FSourceControlMenuHelpers::GetSourceControlCheckInStatusTooltipText)
 			.Visibility_Static(&FSourceControlMenuHelpers::GetSourceControlCheckInStatusVisibility)
-			.IsEnabled_Lambda([]() { return GetNumLocalChanges() > 0; })
+			.IsEnabled_Lambda([]() { return HasLocalChanges(); })
 			[
 				SNew(SHorizontalBox)
 				+ SHorizontalBox::Slot()
diff --git a/Engine/Source/Editor/StatusBar/Private/SourceControlMenuHelpers.h b/Engine/Source/Editor/StatusBar/Private/SourceControlMenuHelpers.h
--- a/Engine/Source/Editor/StatusBar/Private/SourceControlMenuHelpers.h
+++ b/Engine/Source/Editor/StatusBar/Private/SourceControlMenuHelpers.h
@@ -89,6 +89,7 @@ private:
 
 	/** Check-in button */
 	static int GetNumLocalChanges();
+	static bool HasLocalChanges();
 	static EVisibility GetSourceControlCheckInStatusVisibility();
 	static FText GetSourceControlCheckInStatusText();
 	static FText GetSourceControlCheckInStatusTooltipText();
