Revision Control: Add foreground tint for icon in the status bar

diff --git a/Engine/Source/Editor/StatusBar/Private/SourceControlMenuHelpers.cpp b/Engine/Source/Editor/StatusBar/Private/SourceControlMenuHelpers.cpp
--- a/Engine/Source/Editor/StatusBar/Private/SourceControlMenuHelpers.cpp
+++ b/Engine/Source/Editor/StatusBar/Private/SourceControlMenuHelpers.cpp
@@ -25,6 +25,7 @@
 #include "Misc/MessageDialog.h"
 #include "RevisionControlStyle/RevisionControlStyle.h"
 #include "Bookmarks/BookmarkScoped.h"
+#include "Styling/StyleColors.h"
 
 #define LOCTEXT_NAMESPACE "SourceControlCommands"
 
@@ -493,6 +494,7 @@ TSharedRef<SWidget> FSourceControlMenuHelpers::MakeSourceControlStatusWidget()
 {
 	TSharedRef<SLayeredImage> SourceControlIcon =
 		SNew(SLayeredImage)
+		.ColorAndOpacity(FSlateColor::UseForeground())
 		.Image(FRevisionControlStyleManager::Get().GetBrush("RevisionControl.Icon"));
 
 	SourceControlIcon->AddLayer(TAttribute<const FSlateBrush*>::CreateStatic(&FSourceControlMenuHelpers::GetSourceControlIconBadge));
