SkeinUE: Only enable the 'Revert All' option if there is anything to revert.

diff --git a/Engine/Source/Editor/StatusBar/Private/SourceControlMenuHelpers.cpp b/Engine/Source/Editor/StatusBar/Private/SourceControlMenuHelpers.cpp
--- a/Engine/Source/Editor/StatusBar/Private/SourceControlMenuHelpers.cpp
+++ b/Engine/Source/Editor/StatusBar/Private/SourceControlMenuHelpers.cpp
@@ -88,7 +88,7 @@ void FSourceControlCommands::RegisterCommands()
 	ActionList->MapAction(
 		RevertAll,
 		FExecuteAction::CreateStatic(&FSourceControlCommands::RevertAllModifiedFiles_Clicked),
-		FCanExecuteAction::CreateLambda([] { return true; })
+		FCanExecuteAction::CreateStatic(&FSourceControlCommands::RevertAllModifiedFiles_CanExecute)
 	);
 }
 
@@ -156,6 +156,23 @@ void FSourceControlCommands::CheckOutModifiedFiles_Clicked()
 	FEditorFileUtils::PromptForCheckoutAndSave(PackagesToSave, bCheckDirty, bPromptUserToSave);
 }
 
+bool FSourceControlCommands::RevertAllModifiedFiles_CanExecute()
+{
+	// If CHECK-IN CHANGES button is active, the REVERT ALL option should be active.
+
+	if (FUnsavedAssetsTrackerModule::Get().GetUnsavedAssetNum() > 0)
+	{
+		return true;
+	}
+
+	if (FSourceControlWindows::CanChoosePackagesToCheckIn())
+	{
+		return true;
+	}
+
+	return false;
+}
+
 void FSourceControlCommands::RevertAllModifiedFiles_Clicked()
 {
 	FText Message = LOCTEXT("RevertAllModifiedFiles", "Are you sure you want to revert all local changes? By proceeding, your local changes will be discarded, and the state of your last synced snapshot will be restored.");
diff --git a/Engine/Source/Editor/StatusBar/Private/SourceControlMenuHelpers.h b/Engine/Source/Editor/StatusBar/Private/SourceControlMenuHelpers.h
--- a/Engine/Source/Editor/StatusBar/Private/SourceControlMenuHelpers.h
+++ b/Engine/Source/Editor/StatusBar/Private/SourceControlMenuHelpers.h
@@ -29,6 +29,7 @@ private:
 	static void ViewChangelists_Clicked();
 	static bool CheckOutModifiedFiles_CanExecute();
 	static void CheckOutModifiedFiles_Clicked();
+	static bool RevertAllModifiedFiles_CanExecute();
 	static void RevertAllModifiedFiles_Clicked();
 
 public:
