Status Bar: Add a separator after the Unsaved Assets widget

diff --git a/Engine/Source/Editor/StatusBar/Private/SourceControlMenuHelpers.cpp b/Engine/Source/Editor/StatusBar/Private/SourceControlMenuHelpers.cpp
--- a/Engine/Source/Editor/StatusBar/Private/SourceControlMenuHelpers.cpp
+++ b/Engine/Source/Editor/StatusBar/Private/SourceControlMenuHelpers.cpp
@@ -528,9 +528,17 @@ TSharedRef<SWidget> FSourceControlMenuHelpers::MakeSourceControlStatusWidget()
 			//       ensure its position with respect to the source control button.
 			FUnsavedAssetsTrackerModule::Get().MakeUnsavedAssetsStatusBarWidget()
 		]
+		+ SHorizontalBox::Slot()
+		.AutoWidth()
+		.Padding(4.0f, -5.0f) // Intentional negative padding to make the separator cover the whole status bar vertically
+		[
+			SNew(SSeparator)
+			.Thickness(2.0f)
+			.Orientation(EOrientation::Orient_Vertical)
+		]
 		+ SHorizontalBox::Slot() // Check In Changes Button
 		.VAlign(VAlign_Center)
-		.Padding(FMargin(8.0f, 0.0f, 4.0f, 0.0f))
+		.Padding(FMargin(0.0f, 0.0f, 4.0f, 0.0f))
 		.AutoWidth()
 		[
 			SNew(SButton)
@@ -560,11 +568,11 @@ TSharedRef<SWidget> FSourceControlMenuHelpers::MakeSourceControlStatusWidget()
 			]
 			.OnClicked_Static(&FSourceControlMenuHelpers::OnSourceControlCheckInChangesClicked)
 		]
-		+ SHorizontalBox::Slot()
-		.VAlign(VAlign_Center)
+		+ SHorizontalBox::Slot() // Separator after check in button, only visible when the check in button is visible
 		.AutoWidth()
 		[
 			SNew(SSeparator)
+			.Visibility_Static(&FSourceControlMenuHelpers::GetSourceControlCheckInStatusVisibility)
 			.Thickness(1.0)
 			.Orientation(EOrientation::Orient_Vertical)
 		]
