Updated icons for Check-in Changes and Sync Latest status bar buttons

diff --git a/Engine/Source/Developer/SourceControl/Private/RevisionControlStyle/RevisionControlStyle.cpp b/Engine/Source/Developer/SourceControl/Private/RevisionControlStyle/RevisionControlStyle.cpp
--- a/Engine/Source/Developer/SourceControl/Private/RevisionControlStyle/RevisionControlStyle.cpp
+++ b/Engine/Source/Developer/SourceControl/Private/RevisionControlStyle/RevisionControlStyle.cpp
@@ -120,12 +120,11 @@ FDefaultRevisionControlStyle::FDefaultRevisionControlStyle() : FSlateStyleSet(St
 	
 	// Misc Icons
 	Set("RevisionControl.ChangelistsTab", new IMAGE_BRUSH_SVG("Starship/Common/check-circle", CoreStyleConstants::Icon16x16));
-	Set("RevisionControl.StatusBar.AtLatestRevision", new IMAGE_BRUSH_SVG("Starship/Common/SCC_Cloud_AtLatest", CoreStyleConstants::Icon16x16));
-	Set("RevisionControl.StatusBar.NotAtLatestRevision", new IMAGE_BRUSH_SVG("Starship/Common/SCC_Cloud_NotAtLatest", CoreStyleConstants::Icon16x16));
-	Set("RevisionControl.StatusBar.NoLocalChanges", new IMAGE_BRUSH_SVG("Starship/Common/SCC_Cloud_NoLocalChanges", CoreStyleConstants::Icon16x16));
-	Set("RevisionControl.StatusBar.HasLocalChanges", new IMAGE_BRUSH_SVG("Starship/Common/SCC_Cloud_HasLocalChanges", CoreStyleConstants::Icon16x16));
-
-
+	
+	Set("RevisionControl.StatusBar.AtLatestRevision", new IMAGE_BRUSH_SVG("Starship/SourceControl/RC_StatusRemoteUpToDate", CoreStyleConstants::Icon16x16));
+	Set("RevisionControl.StatusBar.NotAtLatestRevision", new IMAGE_BRUSH_SVG("Starship/SourceControl/RC_StatusRemoteDownload", CoreStyleConstants::Icon16x16, FStyleColors::Warning));
+	Set("RevisionControl.StatusBar.NoLocalChanges", new IMAGE_BRUSH_SVG("Starship/SourceControl/RC_StatusLocalUpToDate", CoreStyleConstants::Icon16x16));
+	Set("RevisionControl.StatusBar.HasLocalChanges", new IMAGE_BRUSH_SVG("Starship/SourceControl/RC_StatusLocalUpload", CoreStyleConstants::Icon16x16, FStyleColors::AccentBlue));
 }
 
 FDefaultRevisionControlStyle::~FDefaultRevisionControlStyle()
diff --git a/Engine/Source/Editor/StatusBar/Private/SourceControlMenuHelpers.cpp b/Engine/Source/Editor/StatusBar/Private/SourceControlMenuHelpers.cpp
--- a/Engine/Source/Editor/StatusBar/Private/SourceControlMenuHelpers.cpp
+++ b/Engine/Source/Editor/StatusBar/Private/SourceControlMenuHelpers.cpp
@@ -568,14 +568,6 @@ TSharedRef<SWidget> FSourceControlMenuHelpers::MakeSourceControlStatusWidget()
 			]
 			.OnClicked_Static(&FSourceControlMenuHelpers::OnSourceControlCheckInChangesClicked)
 		]
-		+ SHorizontalBox::Slot() // Separator after check in button, only visible when the check in button is visible
-		.AutoWidth()
-		[
-			SNew(SSeparator)
-			.Visibility_Static(&FSourceControlMenuHelpers::GetSourceControlCheckInStatusVisibility)
-			.Thickness(1.0)
-			.Orientation(EOrientation::Orient_Vertical)
-		]
 		+SHorizontalBox::Slot() // Check In Kebab Combo button
 		.VAlign(VAlign_Center)
 		.AutoWidth()
@@ -587,6 +579,15 @@ TSharedRef<SWidget> FSourceControlMenuHelpers::MakeSourceControlStatusWidget()
 			.Visibility_Static(&FSourceControlMenuHelpers::GetSourceControlCheckInStatusVisibility)
 			.OnGetMenuContent(FOnGetContent::CreateStatic(&FSourceControlMenuHelpers::GenerateCheckInComboButtonContent))
 		]
+		+ SHorizontalBox::Slot()
+		.VAlign(VAlign_Center)
+		.AutoWidth()
+		[
+			SNew(SSeparator)
+			.Visibility_Static(&FSourceControlMenuHelpers::GetSourceControlCheckInStatusVisibility)
+			.Thickness(1.0)
+			.Orientation(EOrientation::Orient_Vertical)
+		]
 		+ SHorizontalBox::Slot() // Sync Latest Button
 		.VAlign(VAlign_Center)
 		.AutoWidth()
@@ -618,6 +619,15 @@ TSharedRef<SWidget> FSourceControlMenuHelpers::MakeSourceControlStatusWidget()
 			]
 			.OnClicked_Static(&FSourceControlMenuHelpers::OnSourceControlSyncClicked)
 		]
+		+ SHorizontalBox::Slot()
+		.VAlign(VAlign_Center)
+		.AutoWidth()
+		[
+			SNew(SSeparator)
+			.Visibility_Static(&FSourceControlMenuHelpers::GetSourceControlCheckInStatusVisibility)
+			.Thickness(1.0)
+			.Orientation(EOrientation::Orient_Vertical)
+		]
 		+ SHorizontalBox::Slot() // Source Control Menu
 		.VAlign(VAlign_Center)
 		.AutoWidth()
