Trigger an autosave if there are any unsaved changes when a user clicks "check-in changes".

diff --git a/Engine/Source/Editor/StatusBar/Private/SourceControlMenuHelpers.cpp b/Engine/Source/Editor/StatusBar/Private/SourceControlMenuHelpers.cpp
--- a/Engine/Source/Editor/StatusBar/Private/SourceControlMenuHelpers.cpp
+++ b/Engine/Source/Editor/StatusBar/Private/SourceControlMenuHelpers.cpp
@@ -179,23 +179,7 @@ void FSourceControlCommands::RevertAllModifiedFiles_Clicked()
 	FText Title = LOCTEXT("RevertAllModifiedFiles_Title", "Revert all local changes");
 	if (FMessageDialog::Open(EAppMsgType::YesNo, EAppReturnType::No, Message, &Title) == EAppReturnType::Yes)
 	{
-		// Get a list of all the unsaved packages
-		TArray<FString> UnsavedFileNames = FUnsavedAssetsTrackerModule::Get().GetUnsavedAssets();
-		TArray<UPackage*> UnsavedPackages;
-		UnsavedPackages.Reserve(UnsavedFileNames.Num());
-
-		for (FString& FileName : UnsavedFileNames)
-		{
-			FString PackageName = UPackageTools::FilenameToPackageName(FileName);
-			UPackage* Package = FindPackage(nullptr, *PackageName);
-			if (Package != nullptr)
-			{
-				UnsavedPackages.Add(Package);
-			}
-		}
-
-		UEditorLoadingAndSavingUtils::SavePackages(UnsavedPackages, /*bOnlyDirty=*/true);
-		
+		FSourceControlMenuHelpers::SaveUnsavedFiles();
 		FBookmarkScoped BookmarkScoped;
 		FSourceControlWindows::RevertAllChangesAndReloadWorld();
 	}
@@ -452,6 +436,26 @@ bool FSourceControlMenuHelpers::HasLocalChanges()
 	return GetNumLocalChanges() > 0 || FUnsavedAssetsTrackerModule::Get().GetUnsavedAssetNum() > 0;
 }
 
+void FSourceControlMenuHelpers::SaveUnsavedFiles()
+{
+	// Get a list of all the unsaved packages
+	TArray<FString> UnsavedFileNames = FUnsavedAssetsTrackerModule::Get().GetUnsavedAssets();
+	TArray<UPackage*> UnsavedPackages;
+	UnsavedPackages.Reserve(UnsavedFileNames.Num());
+
+	for (FString& FileName : UnsavedFileNames)
+	{
+		FString PackageName = UPackageTools::FilenameToPackageName(FileName);
+		UPackage* Package = FindPackage(nullptr, *PackageName);
+		if (Package != nullptr)
+		{
+			UnsavedPackages.Add(Package);
+		}
+	}
+
+	UEditorLoadingAndSavingUtils::SavePackages(UnsavedPackages, /*bOnlyDirty=*/true);
+}
+
 EVisibility FSourceControlMenuHelpers::GetSourceControlCheckInStatusVisibility()
 {
 	bool bDisplaySourceControlCheckInStatus = false;
@@ -499,8 +503,9 @@ const FSlateBrush* FSourceControlMenuHelpers::GetSourceControlCheckInStatusIcon(
 
 FReply FSourceControlMenuHelpers::OnSourceControlCheckInChangesClicked()
 {
-	if (FUnsavedAssetsTrackerModule::Get().PromptToSavePackages() || HasLocalChanges())
+	if (HasLocalChanges())
 	{
+		FSourceControlMenuHelpers::SaveUnsavedFiles();
 		FSourceControlWindows::ChoosePackagesToCheckIn();
 	}
 
diff --git a/Engine/Source/Editor/StatusBar/Private/SourceControlMenuHelpers.h b/Engine/Source/Editor/StatusBar/Private/SourceControlMenuHelpers.h
--- a/Engine/Source/Editor/StatusBar/Private/SourceControlMenuHelpers.h
+++ b/Engine/Source/Editor/StatusBar/Private/SourceControlMenuHelpers.h
@@ -91,6 +91,7 @@ private:
 	/** Check-in button */
 	static int GetNumLocalChanges();
 	static bool HasLocalChanges();
+	static void SaveUnsavedFiles();
 	static EVisibility GetSourceControlCheckInStatusVisibility();
 	static FText GetSourceControlCheckInStatusText();
 	static FText GetSourceControlCheckInStatusTooltipText();
