SourceControl: Don't enable check-in changes button when there are unsaved changes. After those changes are saved to disk, there exists a window of time before the provider realizes there are changes, resulting in the button temporary switching back to 'No Changes' before re-activating 'Check-In Changes' again.
This was deemed problematic, so falling back to old behavior.

diff --git a/Engine/Source/Editor/StatusBar/Private/SourceControlMenuHelpers.cpp b/Engine/Source/Editor/StatusBar/Private/SourceControlMenuHelpers.cpp
--- a/Engine/Source/Editor/StatusBar/Private/SourceControlMenuHelpers.cpp
+++ b/Engine/Source/Editor/StatusBar/Private/SourceControlMenuHelpers.cpp
@@ -360,8 +360,15 @@ const FSlateBrush* FSourceControlMenuHelpers::GetSourceControlIconBadge()
 bool FSourceControlMenuHelpers::IsAtLatestRevision()
 {
 	ISourceControlModule& SourceControlModule = ISourceControlModule::Get();
-	return SourceControlModule.IsEnabled() && SourceControlModule.GetProvider().IsAvailable() && 
-		SourceControlModule.GetProvider().IsAtLatestRevision().IsSet() && SourceControlModule.GetProvider().IsAtLatestRevision().GetValue();
+	return SourceControlModule.IsEnabled() &&
+		SourceControlModule.GetProvider().IsAvailable() &&
+		SourceControlModule.GetProvider().IsAtLatestRevision().IsSet() &&
+		SourceControlModule.GetProvider().IsAtLatestRevision().GetValue();
+}
+
+bool FSourceControlMenuHelpers::CanSourceControlSync()
+{
+	return !IsAtLatestRevision();
 }
 
 EVisibility FSourceControlMenuHelpers::GetSourceControlSyncStatusVisibility()
@@ -369,30 +376,35 @@ EVisibility FSourceControlMenuHelpers::GetSourceControlSyncStatusVisibility()
 	bool bDisplaySourceControlSyncStatus = false;
 	GConfig->GetBool(TEXT("SourceControlSettings"), TEXT("DisplaySourceControlSyncStatus"), bDisplaySourceControlSyncStatus, GEditorIni);
 
-	ISourceControlModule& SourceControlModule = ISourceControlModule::Get();
-	if (bDisplaySourceControlSyncStatus && SourceControlModule.IsEnabled() && SourceControlModule.GetProvider().IsAvailable() && SourceControlModule.GetProvider().IsAtLatestRevision().IsSet())
+	if (bDisplaySourceControlSyncStatus)
 	{
-		return EVisibility::Visible;
+		ISourceControlModule& SourceControlModule = ISourceControlModule::Get();
+		if (SourceControlModule.IsEnabled() &&
+			SourceControlModule.GetProvider().IsAvailable() &&
+			SourceControlModule.GetProvider().IsAtLatestRevision().IsSet()) // Only providers that implement IsAtLatestRevision are supported.
+		{
+			return EVisibility::Visible;
+		}
 	}
 	return EVisibility::Collapsed;
 }
 
 FText FSourceControlMenuHelpers::GetSourceControlSyncStatusText()
 {
-	if (IsAtLatestRevision())
+	if (CanSourceControlSync())
 	{
-		return LOCTEXT("SyncLatestButtonAtHeadText", "At Latest");
+		return LOCTEXT("SyncLatestButtonNotAtHeadText", "Sync Latest");
 	}
-	return LOCTEXT("SyncLatestButtonNotAtHeadText", "Sync Latest");
+	return LOCTEXT("SyncLatestButtonAtHeadText", "At Latest");
 }
 
 FText FSourceControlMenuHelpers::GetSourceControlSyncStatusTooltipText()
 {
-	if (IsAtLatestRevision())
+	if (CanSourceControlSync())
 	{
-		return LOCTEXT("SyncLatestButtonAtHeadTooltipText", "Currently at the latest Snapshot for this project");
+		return LOCTEXT("SyncLatestButtonNotAtHeadTooltipText", "Sync to the latest Snapshot for this project");
 	}
-	return LOCTEXT("SyncLatestButtonNotAtHeadTooltipText", "Sync to the latest Snapshot for this project");
+	return LOCTEXT("SyncLatestButtonAtHeadTooltipText", "Currently at the latest Snapshot for this project");
 }
 
 const FSlateBrush* FSourceControlMenuHelpers::GetSourceControlSyncStatusIcon()
@@ -400,11 +412,11 @@ const FSlateBrush* FSourceControlMenuHelpers::GetSourceControlSyncStatusIcon()
 	static const FSlateBrush* AtHeadBrush = FRevisionControlStyleManager::Get().GetBrush("RevisionControl.StatusBar.AtLatestRevision");
 	static const FSlateBrush* NotAtHeadBrush = FRevisionControlStyleManager::Get().GetBrush("RevisionControl.StatusBar.NotAtLatestRevision");
 
-	if (IsAtLatestRevision())
+	if (CanSourceControlSync())
 	{
-		return AtHeadBrush;
+		return NotAtHeadBrush;
 	}
-	return NotAtHeadBrush;
+	return AtHeadBrush;
 }
 
 FReply FSourceControlMenuHelpers::OnSourceControlSyncClicked()
@@ -423,37 +435,41 @@ FReply FSourceControlMenuHelpers::OnSourceControlSyncClicked()
 int FSourceControlMenuHelpers::GetNumLocalChanges()
 {
 	ISourceControlModule& SourceControlModule = ISourceControlModule::Get();
-	if (SourceControlModule.IsEnabled() && SourceControlModule.GetProvider().IsAvailable()
-		&& SourceControlModule.GetProvider().GetNumLocalChanges().IsSet())
+	if (SourceControlModule.IsEnabled() &&
+		SourceControlModule.GetProvider().IsAvailable() &&
+		SourceControlModule.GetProvider().GetNumLocalChanges().IsSet())
 	{
 		return SourceControlModule.GetProvider().GetNumLocalChanges().GetValue();
 	}
 	return 0;
 }
 
-bool FSourceControlMenuHelpers::HasLocalChanges()
-{
-	return GetNumLocalChanges() > 0 || FUnsavedAssetsTrackerModule::Get().GetUnsavedAssetNum() > 0;
-}
-
 void FSourceControlMenuHelpers::SaveUnsavedFiles()
 {
 	// Get a list of all the unsaved packages
 	TArray<FString> UnsavedFileNames = FUnsavedAssetsTrackerModule::Get().GetUnsavedAssets();
-	TArray<UPackage*> UnsavedPackages;
-	UnsavedPackages.Reserve(UnsavedFileNames.Num());
-
-	for (FString& FileName : UnsavedFileNames)
+	if (UnsavedFileNames.Num() > 0)
 	{
-		FString PackageName = UPackageTools::FilenameToPackageName(FileName);
-		UPackage* Package = FindPackage(nullptr, *PackageName);
-		if (Package != nullptr)
+		TArray<UPackage*> UnsavedPackages;
+		UnsavedPackages.Reserve(UnsavedFileNames.Num());
+
+		for (FString& FileName : UnsavedFileNames)
 		{
-			UnsavedPackages.Add(Package);
+			FString PackageName = UPackageTools::FilenameToPackageName(FileName);
+			UPackage* Package = FindPackage(nullptr, *PackageName);
+			if (Package != nullptr)
+			{
+				UnsavedPackages.Add(Package);
+			}
 		}
+
+		UEditorLoadingAndSavingUtils::SavePackages(UnsavedPackages, /*bOnlyDirty=*/true);
 	}
+}
 
-	UEditorLoadingAndSavingUtils::SavePackages(UnsavedPackages, /*bOnlyDirty=*/true);
+bool FSourceControlMenuHelpers::CanSourceControlCheckIn()
+{
+	return (GetNumLocalChanges() > 0);
 }
 
 EVisibility FSourceControlMenuHelpers::GetSourceControlCheckInStatusVisibility()
@@ -461,17 +477,22 @@ EVisibility FSourceControlMenuHelpers::GetSourceControlCheckInStatusVisibility()
 	bool bDisplaySourceControlCheckInStatus = false;
 	GConfig->GetBool(TEXT("SourceControlSettings"), TEXT("DisplaySourceControlCheckInStatus"), bDisplaySourceControlCheckInStatus, GEditorIni);
 
-	ISourceControlModule& SourceControlModule = ISourceControlModule::Get();
-	if (bDisplaySourceControlCheckInStatus && SourceControlModule.IsEnabled() && SourceControlModule.GetProvider().IsAvailable() && SourceControlModule.GetProvider().GetNumLocalChanges().IsSet())
+	if (bDisplaySourceControlCheckInStatus)
 	{
-		return EVisibility::Visible;
+		ISourceControlModule& SourceControlModule = ISourceControlModule::Get();
+		if (SourceControlModule.IsEnabled() &&
+			SourceControlModule.GetProvider().IsAvailable() &&
+			SourceControlModule.GetProvider().GetNumLocalChanges().IsSet()) // Only providers that implement GetNumLocalChanges are supported.
+		{
+			return EVisibility::Visible;
+		}
 	}
 	return EVisibility::Collapsed;
 }
 
 FText FSourceControlMenuHelpers::GetSourceControlCheckInStatusText()
 {
-	if (HasLocalChanges())
+	if (CanSourceControlCheckIn())
 	{
 		return LOCTEXT("CheckInButtonChangesText", "Check-in Changes");
 	}
@@ -481,10 +502,9 @@ FText FSourceControlMenuHelpers::GetSourceControlCheckInStatusText()
 
 FText FSourceControlMenuHelpers::GetSourceControlCheckInStatusTooltipText()
 {
-	if (HasLocalChanges())
+	if (CanSourceControlCheckIn())
 	{
-		return FText::Format(LOCTEXT("CheckInButtonChangesTooltipText", "Check-in {0} change(s) to this project"), 
-			GetNumLocalChanges() > 0 ? GetNumLocalChanges() : FUnsavedAssetsTrackerModule::Get().GetUnsavedAssetNum());
+		return FText::Format(LOCTEXT("CheckInButtonChangesTooltipText", "Check-in {0} change(s) to this project"), GetNumLocalChanges());
 	}
 	return LOCTEXT("CheckInButtonNoChangesTooltipText", "No Changes to check in for this project");
 }
@@ -494,7 +514,7 @@ const FSlateBrush* FSourceControlMenuHelpers::GetSourceControlCheckInStatusIcon(
 	static const FSlateBrush* NoLocalChangesBrush = FRevisionControlStyleManager::Get().GetBrush("RevisionControl.StatusBar.NoLocalChanges");
 	static const FSlateBrush* HasLocalChangesBrush = FRevisionControlStyleManager::Get().GetBrush("RevisionControl.StatusBar.HasLocalChanges");
 
-	if (HasLocalChanges())
+	if (CanSourceControlCheckIn())
 	{
 		return HasLocalChangesBrush;
 	}
@@ -503,7 +523,7 @@ const FSlateBrush* FSourceControlMenuHelpers::GetSourceControlCheckInStatusIcon(
 
 FReply FSourceControlMenuHelpers::OnSourceControlCheckInChangesClicked()
 {
-	if (HasLocalChanges())
+	if (CanSourceControlCheckIn())
 	{
 		FSourceControlMenuHelpers::SaveUnsavedFiles();
 		FSourceControlWindows::ChoosePackagesToCheckIn();
@@ -550,7 +570,7 @@ TSharedRef<SWidget> FSourceControlMenuHelpers::MakeSourceControlStatusWidget()
 			.ButtonStyle(&FAppStyle::Get().GetWidgetStyle<FButtonStyle>("StatusBar.StatusBarButton"))
 			.ToolTipText_Static(&FSourceControlMenuHelpers::GetSourceControlCheckInStatusTooltipText)
 			.Visibility_Static(&FSourceControlMenuHelpers::GetSourceControlCheckInStatusVisibility)
-			.IsEnabled_Lambda([]() { return HasLocalChanges(); })
+			.IsEnabled_Lambda([]() { return CanSourceControlCheckIn(); })
 			[
 				SNew(SHorizontalBox)
 				+ SHorizontalBox::Slot()
@@ -601,7 +621,7 @@ TSharedRef<SWidget> FSourceControlMenuHelpers::MakeSourceControlStatusWidget()
 			.ButtonStyle(&FAppStyle::Get().GetWidgetStyle<FButtonStyle>("StatusBar.StatusBarButton"))
 			.ToolTipText_Static(&FSourceControlMenuHelpers::GetSourceControlSyncStatusTooltipText)
 			.Visibility_Static(&FSourceControlMenuHelpers::GetSourceControlSyncStatusVisibility)
-			.IsEnabled_Lambda([]() { return !IsAtLatestRevision(); })
+			.IsEnabled_Lambda([]() { return CanSourceControlSync(); })
 			[
 				SNew(SHorizontalBox)
 				+ SHorizontalBox::Slot()
diff --git a/Engine/Source/Editor/StatusBar/Private/SourceControlMenuHelpers.h b/Engine/Source/Editor/StatusBar/Private/SourceControlMenuHelpers.h
--- a/Engine/Source/Editor/StatusBar/Private/SourceControlMenuHelpers.h
+++ b/Engine/Source/Editor/StatusBar/Private/SourceControlMenuHelpers.h
@@ -82,6 +82,7 @@ private:
 
 	/** Sync button */
 	static bool IsAtLatestRevision();
+	static bool CanSourceControlSync();
 	static EVisibility GetSourceControlSyncStatusVisibility();
 	static FText GetSourceControlSyncStatusText();
 	static FText GetSourceControlSyncStatusTooltipText();
@@ -90,8 +91,8 @@ private:
 
 	/** Check-in button */
 	static int GetNumLocalChanges();
-	static bool HasLocalChanges();
 	static void SaveUnsavedFiles();
+	static bool CanSourceControlCheckIn();
 	static EVisibility GetSourceControlCheckInStatusVisibility();
 	static FText GetSourceControlCheckInStatusText();
 	static FText GetSourceControlCheckInStatusTooltipText();
