Avoid crashing when starting progress notification while slate is still uninitialized

#rnx

diff --git a/Engine/Source/Editor/StatusBar/Private/StatusBarSubsystem.cpp b/Engine/Source/Editor/StatusBar/Private/StatusBarSubsystem.cpp
--- a/Engine/Source/Editor/StatusBar/Private/StatusBarSubsystem.cpp
+++ b/Engine/Source/Editor/StatusBar/Private/StatusBarSubsystem.cpp
@@ -627,24 +627,28 @@ void UStatusBarSubsystem::ClearStatusBarMessages(FName StatusBarName)
 
 void UStatusBarSubsystem::StartProgressNotification(FProgressNotificationHandle Handle, FText DisplayText, int32 TotalWorkToDo)
 {
-	// Get the active window, if one is not active a notification was started when the application was deactivated so use the focus path to find a window or just use the root window if there is no keyboard focus
-	TSharedPtr<SWindow> ActiveWindow = FSlateApplication::Get().GetActiveTopLevelRegularWindow();
-	if (!ActiveWindow)
+	// Avoid crashing when starting progress notification while slate is still uninitialized. (i.e. commandlet)
+	if (FSlateApplication::IsInitialized())
 	{
-		TSharedPtr<SWidget> FocusedWidget = FSlateApplication::Get().GetKeyboardFocusedWidget();
-		ActiveWindow = FocusedWidget ? FSlateApplication::Get().FindWidgetWindow(FocusedWidget.ToSharedRef()) : FGlobalTabmanager::Get()->GetRootWindow();
-	}
+		// Get the active window, if one is not active a notification was started when the application was deactivated so use the focus path to find a window or just use the root window if there is no keyboard focus
+		TSharedPtr<SWindow> ActiveWindow = FSlateApplication::Get().GetActiveTopLevelRegularWindow();
+		if (!ActiveWindow)
+		{
+			TSharedPtr<SWidget> FocusedWidget = FSlateApplication::Get().GetKeyboardFocusedWidget();
+			ActiveWindow = FocusedWidget ? FSlateApplication::Get().FindWidgetWindow(FocusedWidget.ToSharedRef()) : FGlobalTabmanager::Get()->GetRootWindow();
+		}
 
-	// Find the active status bar to display the progress in
-	for (auto StatusBar : StatusBars)
-	{
-		if (TSharedPtr<SStatusBar> StatusBarPinned = StatusBar.Value.StatusBarWidget.Pin())
+		// Find the active status bar to display the progress in
+		for (auto StatusBar : StatusBars)
 		{
-			TSharedPtr<SDockTab> ParentTab = StatusBarPinned->GetParentTab();
-			if (ParentTab && ParentTab->IsForeground() && ParentTab->GetParentWindow() == ActiveWindow)
+			if (TSharedPtr<SStatusBar> StatusBarPinned = StatusBar.Value.StatusBarWidget.Pin())
 			{
-				StatusBarPinned->StartProgressNotification(Handle, DisplayText, TotalWorkToDo);
-				break;
+				TSharedPtr<SDockTab> ParentTab = StatusBarPinned->GetParentTab();
+				if (ParentTab && ParentTab->IsForeground() && ParentTab->GetParentWindow() == ActiveWindow)
+				{
+					StatusBarPinned->StartProgressNotification(Handle, DisplayText, TotalWorkToDo);
+					break;
+				}
 			}
 		}
 	}
