SkeinUE: When using the CHECK-IN CHANGES button, perform a SYNC first, as the API expects pull-before-push.

diff --git a/Engine/Source/Editor/SourceControlWindows/Private/SourceControlWindows.cpp b/Engine/Source/Editor/SourceControlWindows/Private/SourceControlWindows.cpp
--- a/Engine/Source/Editor/SourceControlWindows/Private/SourceControlWindows.cpp
+++ b/Engine/Source/Editor/SourceControlWindows/Private/SourceControlWindows.cpp
@@ -182,6 +182,7 @@ bool FSourceControlWindows::SyncLatest()
 			FText Message(LOCTEXT("SCC_Sync_Failed", "Failed to sync files!"));
 			FMessageLog("SourceControl").Notify(Message);
 		}
+		return bSuccess;
 	}
 
 	return false;
diff --git a/Engine/Source/Editor/StatusBar/Private/SourceControlMenuHelpers.cpp b/Engine/Source/Editor/StatusBar/Private/SourceControlMenuHelpers.cpp
--- a/Engine/Source/Editor/StatusBar/Private/SourceControlMenuHelpers.cpp
+++ b/Engine/Source/Editor/StatusBar/Private/SourceControlMenuHelpers.cpp
@@ -528,8 +528,19 @@ FReply FSourceControlMenuHelpers::OnSourceControlCheckInChangesClicked()
 {
 	if (CanSourceControlCheckIn())
 	{
-		FSourceControlMenuHelpers::SaveUnsavedFiles();
-		FSourceControlWindows::ChoosePackagesToCheckIn();
+		bool bSyncNeeded = FSourceControlWindows::CanSyncLatest();
+		bool bSyncSuccess = true;
+		if (bSyncNeeded)
+		{
+			FBookmarkScoped BookmarkScoped;
+			bSyncSuccess = FSourceControlWindows::SyncLatest();
+		}
+
+		if (bSyncSuccess)
+		{
+			FSourceControlMenuHelpers::SaveUnsavedFiles();
+			FSourceControlWindows::ChoosePackagesToCheckIn();
+		}
 	}
 
 	return FReply::Handled();
