Uncontrolled Changelist: Show 'View Changelists' button if provider uses Uncontrolled Changlists (re-enables it for Default provider)

#rnx

diff --git a/Engine/Source/Editor/StatusBar/Private/SourceControlMenuHelpers.cpp b/Engine/Source/Editor/StatusBar/Private/SourceControlMenuHelpers.cpp
--- a/Engine/Source/Editor/StatusBar/Private/SourceControlMenuHelpers.cpp
+++ b/Engine/Source/Editor/StatusBar/Private/SourceControlMenuHelpers.cpp
@@ -106,7 +106,7 @@ bool FSourceControlCommands::ViewChangelists_CanExecute()
 
 bool FSourceControlCommands::ViewChangelists_IsVisible()
 {
-	return ISourceControlModule::Get().GetProvider().UsesChangelists();
+	return ISourceControlModule::Get().GetProvider().UsesChangelists() || ISourceControlModule::Get().GetProvider().UsesUncontrolledChangelists();
 }
 
 bool FSourceControlCommands::SubmitContent_IsVisible()
