Skein / ConflictUE: Add conflict information to StatusBar.

Why? These buttons have to behave differently when there are unresolved conflicts remaining.

Note that it isn't possible to add a build dependency from a module in /Engine to the Skein module in /Engine/Plugins, hence the cvar/ccmd approach.

#rnx

diff --git a/Engine/Source/Developer/SourceControl/Private/RevisionControlStyle/RevisionControlStyle.cpp b/Engine/Source/Developer/SourceControl/Private/RevisionControlStyle/RevisionControlStyle.cpp
--- a/Engine/Source/Developer/SourceControl/Private/RevisionControlStyle/RevisionControlStyle.cpp
+++ b/Engine/Source/Developer/SourceControl/Private/RevisionControlStyle/RevisionControlStyle.cpp
@@ -126,6 +126,7 @@ FDefaultRevisionControlStyle::FDefaultRevisionControlStyle() : FSlateStyleSet(St
 	Set("RevisionControl.StatusBar.NotAtLatestRevision", new IMAGE_BRUSH_SVG("Starship/SourceControl/RC_StatusRemoteDownload", CoreStyleConstants::Icon16x16, FStyleColors::Warning));
 	Set("RevisionControl.StatusBar.NoLocalChanges", new IMAGE_BRUSH_SVG("Starship/SourceControl/RC_StatusLocalUpToDate", CoreStyleConstants::Icon16x16));
 	Set("RevisionControl.StatusBar.HasLocalChanges", new IMAGE_BRUSH_SVG("Starship/SourceControl/RC_StatusLocalUpload", CoreStyleConstants::Icon16x16, FStyleColors::AccentBlue));
+	Set("RevisionControl.StatusBar.Conflicted", new IMAGE_BRUSH_SVG("Starship/SourceControl/RC_Conflicted", CoreStyleConstants::Icon16x16, FStyleColors::Warning));
 }
 
 FDefaultRevisionControlStyle::~FDefaultRevisionControlStyle()
diff --git a/Engine/Source/Editor/StatusBar/Private/SourceControlMenuHelpers.cpp b/Engine/Source/Editor/StatusBar/Private/SourceControlMenuHelpers.cpp
--- a/Engine/Source/Editor/StatusBar/Private/SourceControlMenuHelpers.cpp
+++ b/Engine/Source/Editor/StatusBar/Private/SourceControlMenuHelpers.cpp
@@ -1,12 +1,10 @@
 // Copyright Epic Games, Inc. All Rights Reserved.
 
 #include "SourceControlMenuHelpers.h"
-#include "ISourceControlOperation.h"
 #include "SourceControlOperations.h"
 #include "ISourceControlProvider.h"
 #include "ISourceControlModule.h"
 #include "ISourceControlWindowsModule.h"
-#include "SourceControlMenuHelpers.h"
 #include "SourceControlWindows.h"
 #include "UnsavedAssetsTrackerModule.h"
 #include "FileHelpers.h"
@@ -26,6 +24,7 @@
 #include "RevisionControlStyle/RevisionControlStyle.h"
 #include "Bookmarks/BookmarkScoped.h"
 #include "Styling/StyleColors.h"
+#include "HAL/IConsoleManager.h"
 
 #define LOCTEXT_NAMESPACE "SourceControlCommands"
 
@@ -195,11 +194,58 @@ void FSourceControlCommands::RevertAllModifiedFiles_Clicked()
 	}
 }
 
-FSourceControlMenuHelpers& FSourceControlMenuHelpers::Get()
+FDelegateHandle FSourceControlMenuHelpers::SourceControlProviderChangedHandle;
+FDelegateHandle FSourceControlMenuHelpers::SourceControlStateChangedHandle;
+
+bool FSourceControlMenuHelpers::bConflictsRemaining = false;
+
+void FSourceControlMenuHelpers::OnSourceControlProviderChanged(ISourceControlProvider& OldProvider, ISourceControlProvider& NewProvider)
 {
-	// Singleton instance
-	static FSourceControlMenuHelpers SourceControlMenuHelpers;
-	return SourceControlMenuHelpers;
+	if (SourceControlStateChangedHandle.IsValid())
+	{
+		OldProvider.UnregisterSourceControlStateChanged_Handle(SourceControlStateChangedHandle);
+		SourceControlStateChangedHandle.Reset();
+	}
+
+	if (!IsEngineExitRequested())
+	{
+		SourceControlStateChangedHandle = NewProvider.RegisterSourceControlStateChanged_Handle(
+			FSourceControlStateChanged::FDelegate::CreateStatic(&FSourceControlMenuHelpers::OnSourceControlStateChanged)
+		);
+	}
+}
+
+void FSourceControlMenuHelpers::OnSourceControlStateChanged()
+{
+	ISourceControlProvider& SourceControlProvider = ISourceControlModule::Get().GetProvider();
+
+	TArray<FSourceControlStateRef> Conflicts = SourceControlProvider.GetCachedStateByPredicate(
+		[](const FSourceControlStateRef& State)
+		{
+			return State->IsConflicted();
+		}
+	);
+
+	bConflictsRemaining = (Conflicts.Num() > 0);
+}
+
+bool FSourceControlMenuHelpers::AreConflictsRemaining()
+{
+	const bool bExiting = IsEngineExitRequested();
+
+	if (!bExiting)
+	{
+		// Is ConflictResolution enabled via Unreal Revision Control?
+		if (IConsoleVariable* CVar = IConsoleManager::Get().FindConsoleVariable(TEXT("UnrealRevisionControl.EnableConflictResolution")))
+		{
+			if (CVar->GetBool())
+			{
+				return bConflictsRemaining;
+			}
+		}
+	}
+
+	return false;
 }
 
 FSourceControlMenuHelpers::EQueryState FSourceControlMenuHelpers::QueryState = FSourceControlMenuHelpers::EQueryState::NotQueried;
@@ -214,6 +260,16 @@ void FSourceControlMenuHelpers::CheckSourceControlStatus()
 			FSourceControlOperationComplete::CreateStatic(&FSourceControlMenuHelpers::OnSourceControlOperationComplete));
 		QueryState = EQueryState::Querying;
 	}
+
+	if (!SourceControlProviderChangedHandle.IsValid())
+	{
+		SourceControlProviderChangedHandle = SourceControlModule.RegisterProviderChanged(
+			FSourceControlProviderChanged::FDelegate::CreateStatic(&FSourceControlMenuHelpers::OnSourceControlProviderChanged)
+		);
+		SourceControlStateChangedHandle = SourceControlModule.GetProvider().RegisterSourceControlStateChanged_Handle(
+			FSourceControlStateChanged::FDelegate::CreateStatic(&FSourceControlMenuHelpers::OnSourceControlStateChanged)
+		);
+	}
 }
 
 void FSourceControlMenuHelpers::OnSourceControlOperationComplete(const FSourceControlOperationRef& InOperation, ECommandResult::Type InResult)
@@ -410,6 +466,10 @@ FText FSourceControlMenuHelpers::GetSourceControlSyncStatusText()
 
 FText FSourceControlMenuHelpers::GetSourceControlSyncStatusTooltipText()
 {
+	if (AreConflictsRemaining())
+	{
+		return LOCTEXT("SyncLatestButtonNotAtHeadTooltipTextConflict", "Some of your changes conflict with the latest snapshot of the project. Select how you will resolve each conflict by choosing between your changes or your collaborator's. Then follow the prompt to resolve them.");
+	}
 	if (CanSourceControlSync())
 	{
 		return LOCTEXT("SyncLatestButtonNotAtHeadTooltipText", "Sync to the latest Snapshot for this project");
@@ -419,9 +479,14 @@ FText FSourceControlMenuHelpers::GetSourceControlSyncStatusTooltipText()
 
 const FSlateBrush* FSourceControlMenuHelpers::GetSourceControlSyncStatusIcon()
 {
+	static const FSlateBrush* ConflictBrush = FRevisionControlStyleManager::Get().GetBrush("RevisionControl.StatusBar.Conflicted");
 	static const FSlateBrush* AtHeadBrush = FRevisionControlStyleManager::Get().GetBrush("RevisionControl.StatusBar.AtLatestRevision");
 	static const FSlateBrush* NotAtHeadBrush = FRevisionControlStyleManager::Get().GetBrush("RevisionControl.StatusBar.NotAtLatestRevision");
 
+	if (AreConflictsRemaining())
+	{
+		return ConflictBrush;
+	}
 	if (CanSourceControlSync())
 	{
 		return NotAtHeadBrush;
@@ -431,7 +496,14 @@ const FSlateBrush* FSourceControlMenuHelpers::GetSourceControlSyncStatusIcon()
 
 FReply FSourceControlMenuHelpers::OnSourceControlSyncClicked()
 {
-	if (FSourceControlWindows::CanSyncLatest())
+	if (AreConflictsRemaining())
+	{
+		if (IConsoleObject* CObj = IConsoleManager::Get().FindConsoleObject(TEXT("UnrealRevisionControl.FocusConflictResolution")))
+		{
+			CObj->AsCommand()->Execute(/*Args=*/TArray<FString>(), /*InWorld=*/nullptr, *GLog);
+		}
+	}
+	else if (CanSourceControlSync())
 	{
 		FBookmarkScoped BookmarkScoped;
 		FSourceControlWindows::SyncLatest();
@@ -489,6 +561,10 @@ FText FSourceControlMenuHelpers::GetSourceControlCheckInStatusText()
 
 FText FSourceControlMenuHelpers::GetSourceControlCheckInStatusTooltipText()
 {
+	if (AreConflictsRemaining())
+	{
+		return LOCTEXT("CheckInButtonChangesTooltipTextConflict", "Some of your changes conflict with the latest snapshot of the project. Select how you will resolve each conflict by choosing between your changes or your collaborator's. Then follow the prompt to resolve them.");
+	}
 	if (CanSourceControlCheckIn())
 	{
 		return FText::Format(LOCTEXT("CheckInButtonChangesTooltipText", "Check-in {0} change(s) to this project"), GetNumLocalChanges());
@@ -498,9 +574,14 @@ FText FSourceControlMenuHelpers::GetSourceControlCheckInStatusTooltipText()
 
 const FSlateBrush* FSourceControlMenuHelpers::GetSourceControlCheckInStatusIcon()
 {
+	static const FSlateBrush* ConflictBrush = FRevisionControlStyleManager::Get().GetBrush("RevisionControl.StatusBar.Conflicted");
 	static const FSlateBrush* NoLocalChangesBrush = FRevisionControlStyleManager::Get().GetBrush("RevisionControl.StatusBar.NoLocalChanges");
 	static const FSlateBrush* HasLocalChangesBrush = FRevisionControlStyleManager::Get().GetBrush("RevisionControl.StatusBar.HasLocalChanges");
 
+	if (AreConflictsRemaining())
+	{
+		return ConflictBrush;
+	}
 	if (CanSourceControlCheckIn())
 	{
 		return HasLocalChangesBrush;
@@ -510,7 +591,14 @@ const FSlateBrush* FSourceControlMenuHelpers::GetSourceControlCheckInStatusIcon(
 
 FReply FSourceControlMenuHelpers::OnSourceControlCheckInChangesClicked()
 {
-	if (CanSourceControlCheckIn())
+	if (AreConflictsRemaining())
+	{
+		if (IConsoleObject* CObj = IConsoleManager::Get().FindConsoleObject(TEXT("UnrealRevisionControl.FocusConflictResolution")))
+		{
+			CObj->AsCommand()->Execute(/*Args=*/TArray<FString>(), /*InWorld=*/nullptr, *GLog);
+		}
+	}
+	else if (CanSourceControlCheckIn())
 	{
 		FSourceControlWindows::ChoosePackagesToCheckIn();
 	}
@@ -526,7 +614,6 @@ TSharedRef<SWidget> FSourceControlMenuHelpers::MakeSourceControlStatusWidget()
 		.Image(FRevisionControlStyleManager::Get().GetBrush("RevisionControl.Icon"));
 
 	SourceControlIcon->AddLayer(TAttribute<const FSlateBrush*>::CreateStatic(&FSourceControlMenuHelpers::GetSourceControlIconBadge));
-
 	
 	return
 		SNew(SHorizontalBox)
diff --git a/Engine/Source/Editor/StatusBar/Private/SourceControlMenuHelpers.h b/Engine/Source/Editor/StatusBar/Private/SourceControlMenuHelpers.h
--- a/Engine/Source/Editor/StatusBar/Private/SourceControlMenuHelpers.h
+++ b/Engine/Source/Editor/StatusBar/Private/SourceControlMenuHelpers.h
@@ -49,13 +49,9 @@ public:
 class FSourceControlMenuHelpers
 {
 	friend class FSourceControlCommands;
-public:
-	/**
-	 * Static: Access singleton instance
-	 *
-	 * @return	Reference to the singleton object
-	 */
-	static FSourceControlMenuHelpers& Get();
+
+private:
+	FSourceControlMenuHelpers() {};
 
 private:
 	enum EQueryState
@@ -65,7 +61,6 @@ private:
 		Queried,
 	};
 
-
 	static EQueryState QueryState;
 
 public:
@@ -97,5 +92,20 @@ private:
 	static FText GetSourceControlCheckInStatusTooltipText();
 	static const FSlateBrush* GetSourceControlCheckInStatusIcon();
 	static FReply OnSourceControlCheckInChangesClicked();
+
+	/** Callbacks */
+	static void OnSourceControlProviderChanged(ISourceControlProvider& OldProvider, ISourceControlProvider& NewProvider);
+	static void OnSourceControlStateChanged();
+
+	/** Conflicts */
+	static bool AreConflictsRemaining();
+
+private:
+	/** Delegate handles */
+	static FDelegateHandle SourceControlProviderChangedHandle;
+	static FDelegateHandle SourceControlStateChangedHandle;
+
+	/** Is there a conflict remaining? */
+	static bool bConflictsRemaining;
 };
 
