Editor | Add mechanism for extending the global status bar elements (Content Browser and Output Log buttons)

diff --git a/Engine/Source/Developer/ToolWidgets/Private/SWidgetDrawer.cpp b/Engine/Source/Developer/ToolWidgets/Private/SWidgetDrawer.cpp
--- a/Engine/Source/Developer/ToolWidgets/Private/SWidgetDrawer.cpp
+++ b/Engine/Source/Developer/ToolWidgets/Private/SWidgetDrawer.cpp
@@ -470,6 +470,12 @@ TSharedRef<SWidget> SWidgetDrawer::MakeStatusBarDrawerButton(const FWidgetDrawer
 					.TextStyle(&FAppStyle::Get().GetWidgetStyle<FTextBlockStyle>("NormalText"))
 					.Text(Drawer.ButtonText)
 				]
+				+ SHorizontalBox::Slot()
+				.VAlign(VAlign_Center)
+				.AutoWidth()
+				[
+					Drawer.CustomButtonWidgets
+				]
 			]
 		];
 
diff --git a/Engine/Source/Developer/ToolWidgets/Public/WidgetDrawerConfig.h b/Engine/Source/Developer/ToolWidgets/Public/WidgetDrawerConfig.h
--- a/Engine/Source/Developer/ToolWidgets/Public/WidgetDrawerConfig.h
+++ b/Engine/Source/Developer/ToolWidgets/Public/WidgetDrawerConfig.h
@@ -4,6 +4,7 @@
 
 #include "CoreMinimal.h"
 #include "Framework/SlateDelegates.h"
+#include "Widgets/SBoxPanel.h"
 
 class SWidget;
 
@@ -25,6 +26,7 @@ struct FWidgetDrawerConfig
 	FText ButtonText;
 	FText ToolTipText;
 	const FSlateBrush* Icon = nullptr;
+	TSharedRef<SHorizontalBox> CustomButtonWidgets = SNew(SHorizontalBox);
 
 	bool operator==(const FName& OtherId) const
 	{
diff --git a/Engine/Source/Editor/StatusBar/Private/StatusBarSubsystem.cpp b/Engine/Source/Editor/StatusBar/Private/StatusBarSubsystem.cpp
--- a/Engine/Source/Editor/StatusBar/Private/StatusBarSubsystem.cpp
+++ b/Engine/Source/Editor/StatusBar/Private/StatusBarSubsystem.cpp
@@ -455,6 +455,11 @@ TSharedRef<SWidget> UStatusBarSubsystem::MakeStatusBarWidget(FName StatusBarName
 	ContentBrowserDrawer.ToolTipText = FText::Format(LOCTEXT("StatusBar_ContentBrowserDrawerToolTip", "Opens a temporary content browser above this status which will dismiss when it loses focus ({0})"), FGlobalEditorCommonCommands::Get().OpenContentBrowserDrawer->GetInputText());
 	ContentBrowserDrawer.Icon = FAppStyle::Get().GetBrush("ContentBrowser.TabIcon");
 
+	for (const TUniquePtr<IGlobalStatusBarExtension>& Extension : GlobalStatusBarExtensions)
+	{
+		Extension->ExtendContentBrowserDrawer(ContentBrowserDrawer);
+	}
+
 	StatusBar->RegisterDrawer(MoveTemp(ContentBrowserDrawer));
 
 	FOutputLogModule& OutputLogModule = FModuleManager::Get().LoadModuleChecked<FOutputLogModule>("OutputLog");
@@ -506,6 +511,12 @@ TSharedRef<SWidget> UStatusBarSubsystem::MakeStatusBarWidget(FName StatusBarName
 	OutputLogDrawer.ButtonText = LOCTEXT("StatusBar_OutputLogButton", "Output Log");
 	OutputLogDrawer.ToolTipText = FText::Format(LOCTEXT("StatusBar_OutputLogButtonTip", "Opens the output log drawer. ({0}) cycles between focusing the console command box, opening the output log drawer, and closing it.\nThe output log drawer may also be toggled directly with ({1})"), FGlobalEditorCommonCommands::Get().OpenConsoleCommandBox->GetInputText(), FGlobalEditorCommonCommands::Get().OpenOutputLogDrawer->GetInputText());
 	OutputLogDrawer.Icon = FAppStyle::Get().GetBrush("Log.TabIcon");
+
+	for (const TUniquePtr<IGlobalStatusBarExtension>& Extension : GlobalStatusBarExtensions)
+	{
+		Extension->ExtendOutputLogDrawer(OutputLogDrawer);
+	}
+
 	StatusBar->RegisterDrawer(MoveTemp(OutputLogDrawer));
 
 	// Clean up stale status bars
@@ -624,6 +635,25 @@ void UStatusBarSubsystem::ClearStatusBarMessages(FName StatusBarName)
 	}
 }
 
+IGlobalStatusBarExtension& UStatusBarSubsystem::RegisterGlobalStatusBarExtension(TUniquePtr<IGlobalStatusBarExtension>&& Extension)
+{
+	int32 Index = GlobalStatusBarExtensions.Add(MoveTemp(Extension));
+	// NOTE: It is safe to return this reference because it's stored in a TUniquePtr which is
+	// guaranteed not to change its address if the array reallocates.
+	return *GlobalStatusBarExtensions[Index];
+}
+
+TUniquePtr<IGlobalStatusBarExtension> UStatusBarSubsystem::UnregisterGlobalStatusBarExtension(IGlobalStatusBarExtension* Extension)
+{
+	int32 Index = GlobalStatusBarExtensions.IndexOfByPredicate([Extension](const TUniquePtr<IGlobalStatusBarExtension>& Other)
+	{
+		return &*Other == Extension;
+	});
+	TUniquePtr<IGlobalStatusBarExtension> RemovedExtension = MoveTemp(GlobalStatusBarExtensions[Index]);
+	GlobalStatusBarExtensions.RemoveAtSwap(Index);
+	return RemovedExtension;
+}
+
 void UStatusBarSubsystem::StartProgressNotification(FProgressNotificationHandle Handle, FText DisplayText, int32 TotalWorkToDo)
 {
 	// Avoid crashing when starting progress notification while slate is still uninitialized. (i.e. commandlet)
diff --git a/Engine/Source/Editor/StatusBar/Public/GlobalStatusBarExtension.h b/Engine/Source/Editor/StatusBar/Public/GlobalStatusBarExtension.h
new file mode 100644
--- /dev/null
+++ b/Engine/Source/Editor/StatusBar/Public/GlobalStatusBarExtension.h
@@ -0,0 +1,20 @@
+// Copyright Epic Games, Inc. All Rights Reserved.
+
+#pragma once
+
+#include "Templates/SharedPointer.h"
+
+struct FWidgetDrawerConfig;
+class SStatusBar;
+
+/** Extension point for extending the functionality of status bars across the entire editor. */
+struct IGlobalStatusBarExtension
+{
+	virtual ~IGlobalStatusBarExtension() = default;
+
+	/** Called before the content browser drawer is registered in the status bar. */
+	virtual void ExtendContentBrowserDrawer(FWidgetDrawerConfig& WidgetDrawerConfig) {}
+
+	/** Called before the output log drawer is registered in the status bar. */
+	virtual void ExtendOutputLogDrawer(FWidgetDrawerConfig& WidgetDrawerConfig) {}
+};
diff --git a/Engine/Source/Editor/StatusBar/Public/StatusBarSubsystem.h b/Engine/Source/Editor/StatusBar/Public/StatusBarSubsystem.h
--- a/Engine/Source/Editor/StatusBar/Public/StatusBarSubsystem.h
+++ b/Engine/Source/Editor/StatusBar/Public/StatusBarSubsystem.h
@@ -5,6 +5,7 @@
 #include "CoreMinimal.h"
 #include "EditorSubsystem.h"
 #include "Framework/Notifications/NotificationManager.h"
+#include "GlobalStatusBarExtension.h"
 #include "StatusBarSubsystem.generated.h"
 
 class SStatusBar;
@@ -163,6 +164,23 @@ public:
 	 */
 	void ClearStatusBarMessages(FName StatusBarName);
 
+	/**
+	 * Registers a new global status bar extension.
+	 *
+	 * @param Extension The extension to register
+	 * @return A reference to the extension interface, which can be later used to unregister the extension
+	 */
+	IGlobalStatusBarExtension& RegisterGlobalStatusBarExtension(TUniquePtr<IGlobalStatusBarExtension>&& Extension);
+
+	/**
+	 * Unregisters an existing status bar extension.
+	 *
+	 * @param Extension Reference returned by RegisterGlobalStatusBarExtension containing the extension to unregister
+	 * @return An owned instance of the extension that was registered, or null if the extension was not found
+	 */
+	TUniquePtr<IGlobalStatusBarExtension> UnregisterGlobalStatusBarExtension(IGlobalStatusBarExtension* Extension);
+
+
 private:
 	/** IProgressNotificationHandler interface */
 	virtual void StartProgressNotification(FProgressNotificationHandle Handle, FText DisplayText, int32 TotalWorkToDo) override;
@@ -196,4 +214,5 @@ private:
 	TSharedPtr<SWidget> StatusBarContentBrowser;
 	TSharedPtr<SWidget> StatusBarOutputLog;
 	static int32 MessageHandleCounter;
+	TArray<TUniquePtr<IGlobalStatusBarExtension>> GlobalStatusBarExtensions;
 };
