Skein / ConflictUE: Add 'Reference viewer' icon and functionality to relevant assets.
- Also updated UI to latest Figma designs.
- Also updated texts to latest ProductSpec ones.

#rnx

diff --git a/Engine/Source/Editor/StatusBar/Private/SourceControlMenuHelpers.cpp b/Engine/Source/Editor/StatusBar/Private/SourceControlMenuHelpers.cpp
--- a/Engine/Source/Editor/StatusBar/Private/SourceControlMenuHelpers.cpp
+++ b/Engine/Source/Editor/StatusBar/Private/SourceControlMenuHelpers.cpp
@@ -468,7 +468,7 @@ FText FSourceControlMenuHelpers::GetSourceControlSyncStatusTooltipText()
 {
 	if (AreConflictsRemaining())
 	{
-		return LOCTEXT("SyncLatestButtonNotAtHeadTooltipTextConflict", "Some of your changes conflict with the latest snapshot of the project. Select how you will resolve each conflict by choosing between your changes or your collaborator's. Then follow the prompt to resolve them.");
+		return LOCTEXT("SyncLatestButtonNotAtHeadTooltipTextConflict", "Some of your local changes conflict with the latest snapshot of the project. Click here to review these conflicts.");
 	}
 	if (CanSourceControlSync())
 	{
@@ -563,7 +563,7 @@ FText FSourceControlMenuHelpers::GetSourceControlCheckInStatusTooltipText()
 {
 	if (AreConflictsRemaining())
 	{
-		return LOCTEXT("CheckInButtonChangesTooltipTextConflict", "Some of your changes conflict with the latest snapshot of the project. Select how you will resolve each conflict by choosing between your changes or your collaborator's. Then follow the prompt to resolve them.");
+		return LOCTEXT("CheckInButtonChangesTooltipTextConflict", "Some of your local changes conflict with the latest snapshot of the project. Click here to review these conflicts.");
 	}
 	if (CanSourceControlCheckIn())
 	{
