UStatusBarSubsystem
- Don't close an already opened content drawer when calling UStatusBarSubsystem::OpenContentBrowserDrawer.
- UStatusBarSubsystem::OpenContentBrowserDrawer is calling ToggleContentBrowser which closes an opened content drawer. Will now skip calling Toggle if the content drawer is already opened for the parent window.
- Replaced a bunch of autos.

diff --git a/Engine/Source/Editor/StatusBar/Private/StatusBarSubsystem.cpp b/Engine/Source/Editor/StatusBar/Private/StatusBarSubsystem.cpp
--- a/Engine/Source/Editor/StatusBar/Private/StatusBarSubsystem.cpp
+++ b/Engine/Source/Editor/StatusBar/Private/StatusBarSubsystem.cpp
@@ -271,9 +271,9 @@ bool UStatusBarSubsystem::ToggleDebugConsole(TSharedRef<SWindow> ParentWindow, b
 	
 	const bool bCycleToOutputLogDrawer = OutputLogModule.ShouldCycleToOutputLogDrawer() || bAlwaysToggleDrawer;
 	
-	for (auto StatusBar : StatusBars)
+	for (const TPair<FName, FStatusBarData>& StatusBar : StatusBars)
 	{
-		FStatusBarData& SBData = StatusBar.Value;
+		const FStatusBarData& SBData = StatusBar.Value;
 		if (TSharedPtr<SStatusBar> StatusBarPinned = SBData.StatusBarWidget.Pin())
 		{
 			TSharedPtr<SDockTab> ParentTab = StatusBarPinned->GetParentTab();
@@ -333,8 +333,32 @@ bool UStatusBarSubsystem::OpenContentBrowserDrawer()
 
 	if (ParentWindow.IsValid() && ParentWindow->GetType() == EWindowType::Normal)
 	{
-		TSharedRef<SWindow> WindowRef = ParentWindow.ToSharedRef();
-		return ToggleContentBrowser(ParentWindow.ToSharedRef());
+		bool bDrawerIsAlreadyOpened = false;
+		for (const TPair<FName, FStatusBarData>& StatusBar : StatusBars)
+		{
+			if (TSharedPtr<SStatusBar> StatusBarPinned = StatusBar.Value.StatusBarWidget.Pin())
+			{
+				if (StatusBarPinned->IsDrawerOpened(StatusBarDrawerIds::ContentBrowser))
+				{
+					TSharedPtr<SDockTab> ParentTab = StatusBarPinned->GetParentTab();
+					if (ParentTab && ParentTab->IsForeground() && ParentTab->GetParentWindow() == ParentWindow)
+					{
+						bDrawerIsAlreadyOpened = true;
+						break;
+					}
+				}
+			}
+		}
+
+		if (!bDrawerIsAlreadyOpened)
+		{
+			TSharedRef<SWindow> WindowRef = ParentWindow.ToSharedRef();
+			return ToggleContentBrowser(ParentWindow.ToSharedRef());
+		}
+		else
+		{
+			return true;
+		}
 	}
 
 	return false;
@@ -366,9 +390,9 @@ bool UStatusBarSubsystem::TryToggleDrawer(const FName DrawerId)
 
 	if (ParentWindow.IsValid() && ParentWindow->GetType() == EWindowType::Normal)
 	{
-		for (auto StatusBar : StatusBars)
+		for (const TPair<FName, FStatusBarData>& StatusBar : StatusBars)
 		{
-			FStatusBarData& SBData = StatusBar.Value;
+			const FStatusBarData& SBData = StatusBar.Value;
 			if (TSharedPtr<SStatusBar> StatusBarPinned = SBData.StatusBarWidget.Pin())
 			{
 				TSharedPtr<SDockTab> ParentTab = StatusBarPinned->GetParentTab();
@@ -396,7 +420,7 @@ bool UStatusBarSubsystem::TryToggleDrawer(const FName DrawerId)
 bool UStatusBarSubsystem::ForceDismissDrawer()
 {
 	bool bWasDismissed = false;
-	for (auto StatusBar : StatusBars)
+	for (const TPair<FName, FStatusBarData>& StatusBar : StatusBars)
 	{
 		if (TSharedPtr<SStatusBar> StatusBarPinned = StatusBar.Value.StatusBarWidget.Pin())
 		{
@@ -413,7 +437,7 @@ bool UStatusBarSubsystem::ToggleContentBrowser(TSharedRef<SWindow> ParentWindow)
 
 	SNewUserTipNotification::Dismiss();
 
-	for (auto StatusBar : StatusBars)
+	for (const TPair<FName, FStatusBarData>& StatusBar : StatusBars)
 	{
 		if (TSharedPtr<SStatusBar> StatusBarPinned = StatusBar.Value.StatusBarWidget.Pin())
 		{
@@ -668,7 +692,7 @@ void UStatusBarSubsystem::StartProgressNotification(FProgressNotificationHandle
 		}
 
 		// Find the active status bar to display the progress in
-		for (auto StatusBar : StatusBars)
+		for (const TPair<FName, FStatusBarData>& StatusBar : StatusBars)
 		{
 			if (TSharedPtr<SStatusBar> StatusBarPinned = StatusBar.Value.StatusBarWidget.Pin())
 			{
@@ -685,7 +709,7 @@ void UStatusBarSubsystem::StartProgressNotification(FProgressNotificationHandle
 
 void UStatusBarSubsystem::UpdateProgressNotification(FProgressNotificationHandle Handle, int32 TotalWorkDone, int32 UpdatedTotalWorkToDo, FText UpdatedDisplayText)
 {
-	for (auto StatusBar : StatusBars)
+	for (const TPair<FName, FStatusBarData>& StatusBar : StatusBars)
 	{
 		if (TSharedPtr<SStatusBar> StatusBarPinned = StatusBar.Value.StatusBarWidget.Pin())
 		{
@@ -700,7 +724,7 @@ void UStatusBarSubsystem::UpdateProgressNotification(FProgressNotificationHandle
 
 void UStatusBarSubsystem::CancelProgressNotification(FProgressNotificationHandle Handle)
 {
-	for (auto StatusBar : StatusBars)
+	for (const TPair<FName, FStatusBarData>& StatusBar : StatusBars)
 	{
 		if (TSharedPtr<SStatusBar> StatusBarPinned = StatusBar.Value.StatusBarWidget.Pin())
 		{
@@ -734,7 +758,7 @@ void UStatusBarSubsystem::CreateContentBrowserIfNeeded()
 			[this]() -> TSharedPtr<SDockTab>
 			{
 				UE_LOG(LogStatusBar, Log, TEXT("Looking status bar with open content browser drawer..."))
-				for (auto StatusBar : StatusBars)
+				for (const TPair<FName, FStatusBarData>& StatusBar : StatusBars)
 				{
 					if (TSharedPtr<SStatusBar> StatusBarPinned = StatusBar.Value.StatusBarWidget.Pin())
 					{
@@ -865,7 +889,7 @@ void UStatusBarSubsystem::OnContentBrowserOpened(FName StatusBarWithDrawerName)
 	SNewUserTipNotification::Dismiss();
 
 	// Dismiss any other content browser that is opened when one status bar opens it.  The content browser is a shared resource and shouldn't be in the layout twice
-	for (auto StatusBar : StatusBars)
+	for (const TPair<FName, FStatusBarData>& StatusBar : StatusBars)
 	{
 		if (TSharedPtr<SStatusBar> StatusBarPinned = StatusBar.Value.StatusBarWidget.Pin())
 		{
@@ -899,7 +923,7 @@ void UStatusBarSubsystem::OnContentBrowserDismissed(const TSharedPtr<SWidget>& N
 
 void UStatusBarSubsystem::HandleDeferredOpenContentBrowser(TSharedPtr<SWindow> ParentWindow)
 {
-	for (auto StatusBar : StatusBars)
+	for (const TPair<FName, FStatusBarData>& StatusBar : StatusBars)
 	{
 		if (TSharedPtr<SStatusBar> StatusBarPinned = StatusBar.Value.StatusBarWidget.Pin())
 		{
@@ -928,7 +952,7 @@ TSharedRef<SWidget> UStatusBarSubsystem::OnGetOutputLog()
 void UStatusBarSubsystem::OnOutputLogOpened(FName StatusBarWithDrawerName)
 {
 	// Dismiss any other content browser that is opened when one status bar opens it.  The content browser is a shared resource and shouldn't be in the layout twice
-	for (auto StatusBar : StatusBars)
+	for (const TPair<FName, FStatusBarData>& StatusBar : StatusBars)
 	{
 		if (TSharedPtr<SStatusBar> StatusBarPinned = StatusBar.Value.StatusBarWidget.Pin())
 		{
@@ -955,7 +979,7 @@ void UStatusBarSubsystem::OnOutputLogDismised(const TSharedPtr<SWidget>& NewlyFo
 
 void UStatusBarSubsystem::OnDebugConsoleDrawerClosed()
 {
-	for (auto StatusBar : StatusBars)
+	for (const TPair<FName, FStatusBarData>& StatusBar : StatusBars)
 	{
 		if (TSharedPtr<SStatusBar> StatusBarPinned = StatusBar.Value.StatusBarWidget.Pin())
 		{
