* Created Snapshot History Tab structure
* Created Date Picker Widget
* Created Source Control Controls Widget
* Refactored Status Bar menu to use Source Control Controls Widget

diff --git a/Engine/Source/Developer/SourceControl/Private/SSourceControlControls.cpp b/Engine/Source/Developer/SourceControl/Private/SSourceControlControls.cpp
new file mode 100644
--- /dev/null
+++ b/Engine/Source/Developer/SourceControl/Private/SSourceControlControls.cpp
@@ -0,0 +1,392 @@
+// Copyright Epic Games, Inc. All Rights Reserved.
+
+#include "SSourceControlControls.h"
+#include "ISourceControlProvider.h"
+#include "ISourceControlModule.h"
+#include "Misc/ConfigCacheIni.h"
+#include "RevisionControlStyle/RevisionControlStyle.h"
+#include "Widgets/Input/SButton.h"
+#include "Widgets/Input/SComboBox.h"
+#include "Widgets/Layout/SSeparator.h"
+
+
+#define LOCTEXT_NAMESPACE "SSkeinSourceControlWidgets"
+
+/**
+ * Construct this widget
+ *
+ * @param	InArgs	The declaration data for this widget
+ */
+void SSourceControlControls::Construct(const FArguments& InArgs)
+{
+	OnSyncLatestClicked = InArgs._OnClickedSyncLatest;
+	OnCheckInChangesClicked = InArgs._OnClickedCheckInChanges;
+
+	ChildSlot
+	[
+		SNew(SHorizontalBox)
+		+ SHorizontalBox::Slot() // Check In Changes Button
+		.VAlign(VAlign_Center)
+		.Padding(FMargin(0.0f, 0.0f, 4.0f, 0.0f))
+		.AutoWidth()
+		[
+			SNew(SButton)
+			.ButtonStyle(&FAppStyle::Get().GetWidgetStyle<FButtonStyle>("StatusBar.StatusBarButton"))
+			.ToolTipText(this, &SSourceControlControls::GetSourceControlCheckInStatusTooltipText)
+			.Visibility(this, &SSourceControlControls::GetSourceControlCheckInStatusVisibility)
+			.IsEnabled_Lambda([this]() { return CanSourceControlCheckIn(); })
+			[
+				SNew(SHorizontalBox)
+				+ SHorizontalBox::Slot()
+				.AutoWidth()
+				.VAlign(VAlign_Center)
+				.HAlign(HAlign_Center)
+				[
+					SNew(SImage)
+					.Image(this, &SSourceControlControls::GetSourceControlCheckInStatusIcon)
+				]
+				+ SHorizontalBox::Slot()
+				.AutoWidth()
+				.VAlign(VAlign_Center)
+				.Padding(FMargin(5, 0, 0, 0))
+				[
+					SNew(STextBlock)
+					.TextStyle(&FAppStyle::Get().GetWidgetStyle<FTextBlockStyle>("NormalText"))
+					.Text(this, &SSourceControlControls::GetSourceControlCheckInStatusText)
+				]
+			]
+			.OnClicked(InArgs._OnClickedCheckInChanges)
+		]
+		+SHorizontalBox::Slot() // Check In Kebab Combo button
+		.VAlign(VAlign_Center)
+		.AutoWidth()
+		[
+			SNew(SComboButton)
+			.ContentPadding(FMargin(7.f, 0.f))
+			.ComboButtonStyle(&FAppStyle::Get().GetWidgetStyle<FComboButtonStyle>("StatusBar.StatusBarEllipsisComboButton"))
+			.MenuPlacement(MenuPlacement_AboveAnchor)
+			.Visibility(this, &SSourceControlControls::GetSourceControlCheckInStatusVisibility)
+			.OnGetMenuContent(InArgs._OnGenerateKebabMenu)
+		]
+		+ SHorizontalBox::Slot()
+		.VAlign(VAlign_Center)
+		.AutoWidth()
+		[
+			SNew(SSeparator)
+			.Visibility(this, &SSourceControlControls::GetSourceControlCheckInStatusVisibility)
+			.Thickness(1.0)
+			.Orientation(EOrientation::Orient_Vertical)
+		]
+		+ SHorizontalBox::Slot() // Sync Latest Button
+		.VAlign(VAlign_Center)
+		.AutoWidth()
+		[
+			SNew(SButton)
+			.ButtonStyle(&FAppStyle::Get().GetWidgetStyle<FButtonStyle>("StatusBar.StatusBarButton"))
+			.ToolTipText(this, &SSourceControlControls::GetSourceControlSyncStatusTooltipText)
+			.Visibility(this, &SSourceControlControls::GetSourceControlSyncStatusVisibility)
+			.IsEnabled_Lambda([this]() { return CanSourceControlSync(); })
+			[
+				SNew(SHorizontalBox)
+				+ SHorizontalBox::Slot()
+				.AutoWidth()
+				.VAlign(VAlign_Center)
+				.HAlign(HAlign_Center)
+				[
+					SNew(SImage)
+					.Image(this, &SSourceControlControls::GetSourceControlSyncStatusIcon)
+				]
+				+ SHorizontalBox::Slot()
+				.AutoWidth()
+				.VAlign(VAlign_Center)
+				.Padding(FMargin(5, 0, 0, 0))
+				[
+					SNew(STextBlock)
+					.TextStyle(&FAppStyle::Get().GetWidgetStyle<FTextBlockStyle>("NormalText"))
+					.Text(this, &SSourceControlControls::GetSourceControlSyncStatusText)
+				]
+			]
+			.OnClicked(this, &SSourceControlControls::OnSourceControlSyncClicked)
+		]
+		+ SHorizontalBox::Slot()
+		.VAlign(VAlign_Center)
+		.AutoWidth()
+		[
+			SNew(SSeparator)
+			.Visibility(this, &SSourceControlControls::GetSourceControlCheckInStatusVisibility)
+			.Thickness(1.0)
+			.Orientation(EOrientation::Orient_Vertical)
+		]
+	];
+
+	CheckSourceControlStatus();
+}
+
+void SSourceControlControls::CheckSourceControlStatus()
+{
+	ISourceControlModule& SourceControlModule = ISourceControlModule::Get();
+
+	if (!SourceControlProviderChangedHandle.IsValid())
+	{
+		SourceControlProviderChangedHandle = SourceControlModule.RegisterProviderChanged(
+			FSourceControlProviderChanged::FDelegate::CreateSP(this, &SSourceControlControls::OnSourceControlProviderChanged)
+		);
+		SourceControlStateChangedHandle = SourceControlModule.GetProvider().RegisterSourceControlStateChanged_Handle(
+			FSourceControlStateChanged::FDelegate::CreateSP(this, &SSourceControlControls::OnSourceControlStateChanged)
+		);
+	}
+}
+
+void SSourceControlControls::OnSourceControlProviderChanged(ISourceControlProvider& OldProvider, ISourceControlProvider& NewProvider)
+{
+	if (SourceControlStateChangedHandle.IsValid())
+	{
+		OldProvider.UnregisterSourceControlStateChanged_Handle(SourceControlStateChangedHandle);
+		SourceControlStateChangedHandle.Reset();
+	}
+
+	if (!IsEngineExitRequested())
+	{
+		SourceControlStateChangedHandle = NewProvider.RegisterSourceControlStateChanged_Handle(
+			FSourceControlStateChanged::FDelegate::CreateSP(this, &SSourceControlControls::OnSourceControlStateChanged)
+		);
+	}
+}
+
+void SSourceControlControls::OnSourceControlStateChanged()
+{
+	ISourceControlProvider& SourceControlProvider = ISourceControlModule::Get().GetProvider();
+
+	TArray<FSourceControlStateRef> Conflicts = SourceControlProvider.GetCachedStateByPredicate(
+		[](const FSourceControlStateRef& State)
+		{
+			return State->IsConflicted();
+		}
+	);
+
+	bConflictsRemaining = (Conflicts.Num() > 0);
+}
+
+bool SSourceControlControls::AreConflictsRemaining() const
+{
+	const bool bExiting = IsEngineExitRequested();
+
+	if (!bExiting)
+	{
+		// Is ConflictResolution enabled via Unreal Revision Control?
+		if (IConsoleVariable* CVar = IConsoleManager::Get().FindConsoleVariable(TEXT("UnrealRevisionControl.EnableConflictResolution")))
+		{
+			if (CVar->GetBool())
+			{
+				return bConflictsRemaining;
+			}
+		}
+	}
+
+	return false;
+}
+
+/** Sync Status */
+
+bool SSourceControlControls::IsAtLatestRevision() const
+{
+	ISourceControlModule& SourceControlModule = ISourceControlModule::Get();
+	return SourceControlModule.IsEnabled() &&
+		SourceControlModule.GetProvider().IsAvailable() &&
+		SourceControlModule.GetProvider().IsAtLatestRevision().IsSet() &&
+		SourceControlModule.GetProvider().IsAtLatestRevision().GetValue();
+}
+
+bool SSourceControlControls::CanSourceControlSync() const
+{
+	return !IsAtLatestRevision();
+}
+
+EVisibility SSourceControlControls::GetSourceControlSyncStatusVisibility() const
+{
+	if (!GIsEditor)
+	{
+		// Always visible in the Slate Viewer
+		return EVisibility::Visible;
+	}
+
+	bool bDisplaySourceControlSyncStatus = false;
+	GConfig->GetBool(TEXT("SourceControlSettings"), TEXT("DisplaySourceControlSyncStatus"), bDisplaySourceControlSyncStatus, GEditorIni);
+
+	if (bDisplaySourceControlSyncStatus)
+	{
+		ISourceControlModule& SourceControlModule = ISourceControlModule::Get();
+		if (SourceControlModule.IsEnabled() &&
+			SourceControlModule.GetProvider().IsAvailable() &&
+			SourceControlModule.GetProvider().IsAtLatestRevision().IsSet()) // Only providers that implement IsAtLatestRevision are supported.
+		{
+			return EVisibility::Visible;
+		}
+	}
+	return EVisibility::Collapsed;
+}
+
+FText SSourceControlControls::GetSourceControlSyncStatusText() const
+{
+	if (CanSourceControlSync())
+	{
+		return LOCTEXT("SyncLatestButtonNotAtHeadText", "Sync Latest");
+	}
+	return LOCTEXT("SyncLatestButtonAtHeadText", "At Latest");
+}
+
+FText SSourceControlControls::GetSourceControlSyncStatusTooltipText() const
+{
+	if (AreConflictsRemaining())
+	{
+		return LOCTEXT("SyncLatestButtonNotAtHeadTooltipTextConflict", "Some of your local changes conflict with the latest snapshot of the project. Click here to review these conflicts.");
+	}
+	if (CanSourceControlSync())
+	{
+		return LOCTEXT("SyncLatestButtonNotAtHeadTooltipText", "Sync to the latest Snapshot for this project");
+	}
+	return LOCTEXT("SyncLatestButtonAtHeadTooltipText", "Currently at the latest Snapshot for this project");
+}
+
+const FSlateBrush* SSourceControlControls::GetSourceControlSyncStatusIcon() const
+{
+	static const FSlateBrush* ConflictBrush = FRevisionControlStyleManager::Get().GetBrush("RevisionControl.StatusBar.Conflicted");
+	static const FSlateBrush* AtHeadBrush = FRevisionControlStyleManager::Get().GetBrush("RevisionControl.StatusBar.AtLatestRevision");
+	static const FSlateBrush* NotAtHeadBrush = FRevisionControlStyleManager::Get().GetBrush("RevisionControl.StatusBar.NotAtLatestRevision");
+
+	if (AreConflictsRemaining())
+	{
+		return ConflictBrush;
+	}
+	if (CanSourceControlSync())
+	{
+		return NotAtHeadBrush;
+	}
+	return AtHeadBrush;
+}
+
+FReply SSourceControlControls::OnSourceControlSyncClicked() const
+{
+	if (AreConflictsRemaining())
+	{
+		if (IConsoleObject* CObj = IConsoleManager::Get().FindConsoleObject(TEXT("UnrealRevisionControl.FocusConflictResolution")))
+		{
+			CObj->AsCommand()->Execute(/*Args=*/TArray<FString>(), /*InWorld=*/nullptr, *GLog);
+		}
+	}
+	else if (CanSourceControlSync())
+	{
+		if (OnSyncLatestClicked.IsBound())
+		{
+			OnSyncLatestClicked.Execute();
+		}
+	}
+
+	return FReply::Handled();
+}
+
+/** Check-in Status */
+
+int SSourceControlControls::GetNumLocalChanges() const
+{
+	ISourceControlModule& SourceControlModule = ISourceControlModule::Get();
+	if (SourceControlModule.IsEnabled() &&
+		SourceControlModule.GetProvider().IsAvailable() &&
+		SourceControlModule.GetProvider().GetNumLocalChanges().IsSet())
+	{
+		return SourceControlModule.GetProvider().GetNumLocalChanges().GetValue();
+	}
+	return 0;
+}
+
+bool SSourceControlControls::CanSourceControlCheckIn() const
+{
+	return (GetNumLocalChanges() > 0);
+}
+
+EVisibility SSourceControlControls::GetSourceControlCheckInStatusVisibility() const
+{
+	if (!GIsEditor)
+	{
+		// Always visible in the Slate Viewer
+		return EVisibility::Visible;
+	}
+
+	bool bDisplaySourceControlCheckInStatus = false;
+	GConfig->GetBool(TEXT("SourceControlSettings"), TEXT("DisplaySourceControlCheckInStatus"), bDisplaySourceControlCheckInStatus, GEditorIni);
+
+	if (bDisplaySourceControlCheckInStatus)
+	{
+		ISourceControlModule& SourceControlModule = ISourceControlModule::Get();
+		if (SourceControlModule.IsEnabled() &&
+			SourceControlModule.GetProvider().IsAvailable() &&
+			SourceControlModule.GetProvider().GetNumLocalChanges().IsSet()) // Only providers that implement GetNumLocalChanges are supported.
+		{
+			return EVisibility::Visible;
+		}
+	}
+	return EVisibility::Collapsed;
+}
+
+FText SSourceControlControls::GetSourceControlCheckInStatusText() const
+{
+	if (CanSourceControlCheckIn())
+	{
+		return LOCTEXT("CheckInButtonChangesText", "Check-in Changes");
+	}
+	return LOCTEXT("CheckInButtonNoChangesText", "No Changes");
+}
+
+
+FText SSourceControlControls::GetSourceControlCheckInStatusTooltipText() const
+{
+	if (AreConflictsRemaining())
+	{
+		return LOCTEXT("CheckInButtonChangesTooltipTextConflict", "Some of your local changes conflict with the latest snapshot of the project. Click here to review these conflicts.");
+	}
+	if (CanSourceControlCheckIn())
+	{
+		return FText::Format(LOCTEXT("CheckInButtonChangesTooltipText", "Check-in {0} change(s) to this project"), GetNumLocalChanges());
+	}
+	return LOCTEXT("CheckInButtonNoChangesTooltipText", "No Changes to check in for this project");
+}
+
+const FSlateBrush* SSourceControlControls::GetSourceControlCheckInStatusIcon() const
+{
+	static const FSlateBrush* ConflictBrush = FRevisionControlStyleManager::Get().GetBrush("RevisionControl.StatusBar.Conflicted");
+	static const FSlateBrush* NoLocalChangesBrush = FRevisionControlStyleManager::Get().GetBrush("RevisionControl.StatusBar.NoLocalChanges");
+	static const FSlateBrush* HasLocalChangesBrush = FRevisionControlStyleManager::Get().GetBrush("RevisionControl.StatusBar.HasLocalChanges");
+
+	if (AreConflictsRemaining())
+	{
+		return ConflictBrush;
+	}
+	if (CanSourceControlCheckIn())
+	{
+		return HasLocalChangesBrush;
+	}
+	return NoLocalChangesBrush;
+}
+
+FReply SSourceControlControls::OnSourceControlCheckInChangesClicked() const
+{
+	if (AreConflictsRemaining())
+	{
+		if (IConsoleObject* CObj = IConsoleManager::Get().FindConsoleObject(TEXT("UnrealRevisionControl.FocusConflictResolution")))
+		{
+			CObj->AsCommand()->Execute(/*Args=*/TArray<FString>(), /*InWorld=*/nullptr, *GLog);
+		}
+	}
+	else if (CanSourceControlCheckIn())
+	{
+		if (OnCheckInChangesClicked.IsBound())
+		{
+			OnCheckInChangesClicked.Execute();
+		}
+	}
+
+	return FReply::Handled();
+}
+
+
+#undef LOCTEXT_NAMESPACE
\ No newline at end of file
diff --git a/Engine/Source/Developer/SourceControl/Public/SSourceControlControls.h b/Engine/Source/Developer/SourceControl/Public/SSourceControlControls.h
new file mode 100644
--- /dev/null
+++ b/Engine/Source/Developer/SourceControl/Public/SSourceControlControls.h
@@ -0,0 +1,64 @@
+// Copyright Epic Games, Inc. All Rights Reserved.
+
+#pragma once
+
+#include "CoreMinimal.h"
+#include "ISourceControlProvider.h"
+#include "Delegates/IDelegateInstance.h"
+#include "Framework/SlateDelegates.h"
+#include "Styling/SlateTypes.h"
+#include "Widgets/SCompoundWidget.h"
+
+/** Widget for displaying Source Control Check in Changes and Sync Latest buttons */
+class SOURCECONTROL_API SSourceControlControls : public SCompoundWidget
+{
+public:
+	SLATE_BEGIN_ARGS(SSourceControlControls) {}
+
+		SLATE_EVENT(FOnClicked, OnClickedSyncLatest)
+		SLATE_EVENT(FOnClicked, OnClickedCheckInChanges)
+		SLATE_EVENT(FOnGetContent, OnGenerateKebabMenu)
+
+	SLATE_END_ARGS()
+
+public:
+	/** Construct this widget */
+	void Construct(const FArguments& InArgs);
+
+private:
+	/** Sync button */
+	bool IsAtLatestRevision() const;
+	bool CanSourceControlSync() const;
+	EVisibility GetSourceControlSyncStatusVisibility() const;
+	FText GetSourceControlSyncStatusText() const;
+	FText GetSourceControlSyncStatusTooltipText() const;
+	const FSlateBrush* GetSourceControlSyncStatusIcon() const;
+	FReply OnSourceControlSyncClicked() const;
+
+	/** Check-in button */
+	int GetNumLocalChanges() const;
+	bool CanSourceControlCheckIn() const;
+	EVisibility GetSourceControlCheckInStatusVisibility() const;
+	FText GetSourceControlCheckInStatusText() const;
+	FText GetSourceControlCheckInStatusTooltipText() const;
+	const FSlateBrush* GetSourceControlCheckInStatusIcon() const;
+	FReply OnSourceControlCheckInChangesClicked() const;
+
+	/** Conflicts */
+	void CheckSourceControlStatus();
+	bool AreConflictsRemaining() const;
+
+	void OnSourceControlProviderChanged(ISourceControlProvider& OldProvider, ISourceControlProvider& NewProvider);
+	void OnSourceControlStateChanged();
+	
+private:
+	
+	FOnClicked OnSyncLatestClicked;
+	FOnClicked OnCheckInChangesClicked;
+
+	/** Is there a conflict remaining? */
+	bool bConflictsRemaining;
+
+	FDelegateHandle SourceControlProviderChangedHandle;
+	FDelegateHandle SourceControlStateChangedHandle;
+};
\ No newline at end of file
diff --git a/Engine/Source/Editor/StatusBar/Private/SourceControlMenuHelpers.cpp b/Engine/Source/Editor/StatusBar/Private/SourceControlMenuHelpers.cpp
--- a/Engine/Source/Editor/StatusBar/Private/SourceControlMenuHelpers.cpp
+++ b/Engine/Source/Editor/StatusBar/Private/SourceControlMenuHelpers.cpp
@@ -25,6 +25,7 @@
 #include "Bookmarks/BookmarkScoped.h"
 #include "Styling/StyleColors.h"
 #include "HAL/IConsoleManager.h"
+#include "SSourceControlControls.h"
 
 #define LOCTEXT_NAMESPACE "SourceControlCommands"
 
@@ -194,60 +195,6 @@ void FSourceControlCommands::RevertAllModifiedFiles_Clicked()
 	}
 }
 
-FDelegateHandle FSourceControlMenuHelpers::SourceControlProviderChangedHandle;
-FDelegateHandle FSourceControlMenuHelpers::SourceControlStateChangedHandle;
-
-bool FSourceControlMenuHelpers::bConflictsRemaining = false;
-
-void FSourceControlMenuHelpers::OnSourceControlProviderChanged(ISourceControlProvider& OldProvider, ISourceControlProvider& NewProvider)
-{
-	if (SourceControlStateChangedHandle.IsValid())
-	{
-		OldProvider.UnregisterSourceControlStateChanged_Handle(SourceControlStateChangedHandle);
-		SourceControlStateChangedHandle.Reset();
-	}
-
-	if (!IsEngineExitRequested())
-	{
-		SourceControlStateChangedHandle = NewProvider.RegisterSourceControlStateChanged_Handle(
-			FSourceControlStateChanged::FDelegate::CreateStatic(&FSourceControlMenuHelpers::OnSourceControlStateChanged)
-		);
-	}
-}
-
-void FSourceControlMenuHelpers::OnSourceControlStateChanged()
-{
-	ISourceControlProvider& SourceControlProvider = ISourceControlModule::Get().GetProvider();
-
-	TArray<FSourceControlStateRef> Conflicts = SourceControlProvider.GetCachedStateByPredicate(
-		[](const FSourceControlStateRef& State)
-		{
-			return State->IsConflicted();
-		}
-	);
-
-	bConflictsRemaining = (Conflicts.Num() > 0);
-}
-
-bool FSourceControlMenuHelpers::AreConflictsRemaining()
-{
-	const bool bExiting = IsEngineExitRequested();
-
-	if (!bExiting)
-	{
-		// Is ConflictResolution enabled via Unreal Revision Control?
-		if (IConsoleVariable* CVar = IConsoleManager::Get().FindConsoleVariable(TEXT("UnrealRevisionControl.EnableConflictResolution")))
-		{
-			if (CVar->GetBool())
-			{
-				return bConflictsRemaining;
-			}
-		}
-	}
-
-	return false;
-}
-
 FSourceControlMenuHelpers::EQueryState FSourceControlMenuHelpers::QueryState = FSourceControlMenuHelpers::EQueryState::NotQueried;
 
 void FSourceControlMenuHelpers::CheckSourceControlStatus()
@@ -260,16 +207,6 @@ void FSourceControlMenuHelpers::CheckSourceControlStatus()
 			FSourceControlOperationComplete::CreateStatic(&FSourceControlMenuHelpers::OnSourceControlOperationComplete));
 		QueryState = EQueryState::Querying;
 	}
-
-	if (!SourceControlProviderChangedHandle.IsValid())
-	{
-		SourceControlProviderChangedHandle = SourceControlModule.RegisterProviderChanged(
-			FSourceControlProviderChanged::FDelegate::CreateStatic(&FSourceControlMenuHelpers::OnSourceControlProviderChanged)
-		);
-		SourceControlStateChangedHandle = SourceControlModule.GetProvider().RegisterSourceControlStateChanged_Handle(
-			FSourceControlStateChanged::FDelegate::CreateStatic(&FSourceControlMenuHelpers::OnSourceControlStateChanged)
-		);
-	}
 }
 
 void FSourceControlMenuHelpers::OnSourceControlOperationComplete(const FSourceControlOperationRef& InOperation, ECommandResult::Type InResult)
@@ -421,114 +358,19 @@ const FSlateBrush* FSourceControlMenuHelpers::GetSourceControlIconBadge()
 	}
 }
 
-/** Sync Status */
-
-bool FSourceControlMenuHelpers::IsAtLatestRevision()
-{
-	ISourceControlModule& SourceControlModule = ISourceControlModule::Get();
-	return SourceControlModule.IsEnabled() &&
-		SourceControlModule.GetProvider().IsAvailable() &&
-		SourceControlModule.GetProvider().IsAtLatestRevision().IsSet() &&
-		SourceControlModule.GetProvider().IsAtLatestRevision().GetValue();
-}
-
-bool FSourceControlMenuHelpers::CanSourceControlSync()
-{
-	return !IsAtLatestRevision();
-}
-
-EVisibility FSourceControlMenuHelpers::GetSourceControlSyncStatusVisibility()
-{
-	bool bDisplaySourceControlSyncStatus = false;
-	GConfig->GetBool(TEXT("SourceControlSettings"), TEXT("DisplaySourceControlSyncStatus"), bDisplaySourceControlSyncStatus, GEditorIni);
-
-	if (bDisplaySourceControlSyncStatus)
-	{
-		ISourceControlModule& SourceControlModule = ISourceControlModule::Get();
-		if (SourceControlModule.IsEnabled() &&
-			SourceControlModule.GetProvider().IsAvailable() &&
-			SourceControlModule.GetProvider().IsAtLatestRevision().IsSet()) // Only providers that implement IsAtLatestRevision are supported.
-		{
-			return EVisibility::Visible;
-		}
-	}
-	return EVisibility::Collapsed;
-}
-
-FText FSourceControlMenuHelpers::GetSourceControlSyncStatusText()
-{
-	if (CanSourceControlSync())
-	{
-		return LOCTEXT("SyncLatestButtonNotAtHeadText", "Sync Latest");
-	}
-	return LOCTEXT("SyncLatestButtonAtHeadText", "At Latest");
-}
-
-FText FSourceControlMenuHelpers::GetSourceControlSyncStatusTooltipText()
-{
-	if (AreConflictsRemaining())
-	{
-		return LOCTEXT("SyncLatestButtonNotAtHeadTooltipTextConflict", "Some of your local changes conflict with the latest snapshot of the project. Click here to review these conflicts.");
-	}
-	if (CanSourceControlSync())
-	{
-		return LOCTEXT("SyncLatestButtonNotAtHeadTooltipText", "Sync to the latest Snapshot for this project");
-	}
-	return LOCTEXT("SyncLatestButtonAtHeadTooltipText", "Currently at the latest Snapshot for this project");
-}
-
-const FSlateBrush* FSourceControlMenuHelpers::GetSourceControlSyncStatusIcon()
-{
-	static const FSlateBrush* ConflictBrush = FRevisionControlStyleManager::Get().GetBrush("RevisionControl.StatusBar.Conflicted");
-	static const FSlateBrush* AtHeadBrush = FRevisionControlStyleManager::Get().GetBrush("RevisionControl.StatusBar.AtLatestRevision");
-	static const FSlateBrush* NotAtHeadBrush = FRevisionControlStyleManager::Get().GetBrush("RevisionControl.StatusBar.NotAtLatestRevision");
-
-	if (AreConflictsRemaining())
-	{
-		return ConflictBrush;
-	}
-	if (CanSourceControlSync())
-	{
-		return NotAtHeadBrush;
-	}
-	return AtHeadBrush;
-}
-
 FReply FSourceControlMenuHelpers::OnSourceControlSyncClicked()
 {
-	if (AreConflictsRemaining())
-	{
-		if (IConsoleObject* CObj = IConsoleManager::Get().FindConsoleObject(TEXT("UnrealRevisionControl.FocusConflictResolution")))
-		{
-			CObj->AsCommand()->Execute(/*Args=*/TArray<FString>(), /*InWorld=*/nullptr, *GLog);
-		}
-	}
-	else if (CanSourceControlSync())
-	{
-		FBookmarkScoped BookmarkScoped;
-		FSourceControlWindows::SyncLatest();
-	}
+	FBookmarkScoped BookmarkScoped;
+	FSourceControlWindows::SyncLatest();
 
 	return FReply::Handled();
 }
 
-/** Check-in Status */
-
-int FSourceControlMenuHelpers::GetNumLocalChanges()
+FReply FSourceControlMenuHelpers::OnSourceControlCheckInChangesClicked()
 {
-	ISourceControlModule& SourceControlModule = ISourceControlModule::Get();
-	if (SourceControlModule.IsEnabled() &&
-		SourceControlModule.GetProvider().IsAvailable() &&
-		SourceControlModule.GetProvider().GetNumLocalChanges().IsSet())
-	{
-		return SourceControlModule.GetProvider().GetNumLocalChanges().GetValue();
-	}
-	return 0;
-}
+	FSourceControlWindows::ChoosePackagesToCheckIn();
 
-bool FSourceControlMenuHelpers::CanSourceControlCheckIn()
-{
-	return (GetNumLocalChanges() > 0);
+	return FReply::Handled();
 }
 
 EVisibility FSourceControlMenuHelpers::GetSourceControlCheckInStatusVisibility()
@@ -549,63 +391,6 @@ EVisibility FSourceControlMenuHelpers::GetSourceControlCheckInStatusVisibility()
 	return EVisibility::Collapsed;
 }
 
-FText FSourceControlMenuHelpers::GetSourceControlCheckInStatusText()
-{
-	if (CanSourceControlCheckIn())
-	{
-		return LOCTEXT("CheckInButtonChangesText", "Check-in Changes");
-	}
-	return LOCTEXT("CheckInButtonNoChangesText", "No Changes");
-}
-
-
-FText FSourceControlMenuHelpers::GetSourceControlCheckInStatusTooltipText()
-{
-	if (AreConflictsRemaining())
-	{
-		return LOCTEXT("CheckInButtonChangesTooltipTextConflict", "Some of your local changes conflict with the latest snapshot of the project. Click here to review these conflicts.");
-	}
-	if (CanSourceControlCheckIn())
-	{
-		return FText::Format(LOCTEXT("CheckInButtonChangesTooltipText", "Check-in {0} change(s) to this project"), GetNumLocalChanges());
-	}
-	return LOCTEXT("CheckInButtonNoChangesTooltipText", "No Changes to check in for this project");
-}
-
-const FSlateBrush* FSourceControlMenuHelpers::GetSourceControlCheckInStatusIcon()
-{
-	static const FSlateBrush* ConflictBrush = FRevisionControlStyleManager::Get().GetBrush("RevisionControl.StatusBar.Conflicted");
-	static const FSlateBrush* NoLocalChangesBrush = FRevisionControlStyleManager::Get().GetBrush("RevisionControl.StatusBar.NoLocalChanges");
-	static const FSlateBrush* HasLocalChangesBrush = FRevisionControlStyleManager::Get().GetBrush("RevisionControl.StatusBar.HasLocalChanges");
-
-	if (AreConflictsRemaining())
-	{
-		return ConflictBrush;
-	}
-	if (CanSourceControlCheckIn())
-	{
-		return HasLocalChangesBrush;
-	}
-	return NoLocalChangesBrush;
-}
-
-FReply FSourceControlMenuHelpers::OnSourceControlCheckInChangesClicked()
-{
-	if (AreConflictsRemaining())
-	{
-		if (IConsoleObject* CObj = IConsoleManager::Get().FindConsoleObject(TEXT("UnrealRevisionControl.FocusConflictResolution")))
-		{
-			CObj->AsCommand()->Execute(/*Args=*/TArray<FString>(), /*InWorld=*/nullptr, *GLog);
-		}
-	}
-	else if (CanSourceControlCheckIn())
-	{
-		FSourceControlWindows::ChoosePackagesToCheckIn();
-	}
-
-	return FReply::Handled();
-}
-
 TSharedRef<SWidget> FSourceControlMenuHelpers::MakeSourceControlStatusWidget()
 {
 	TSharedRef<SLayeredImage> SourceControlIcon =
@@ -634,97 +419,13 @@ TSharedRef<SWidget> FSourceControlMenuHelpers::MakeSourceControlStatusWidget()
 			.Thickness(2.0f)
 			.Orientation(EOrientation::Orient_Vertical)
 		]
-		+ SHorizontalBox::Slot() // Check In Changes Button
-		.VAlign(VAlign_Center)
-		.Padding(FMargin(0.0f, 0.0f, 4.0f, 0.0f))
-		.AutoWidth()
+		+SHorizontalBox::Slot()
+		.Padding(0.f)
 		[
-			SNew(SButton)
-			.ButtonStyle(&FAppStyle::Get().GetWidgetStyle<FButtonStyle>("StatusBar.StatusBarButton"))
-			.ToolTipText_Static(&FSourceControlMenuHelpers::GetSourceControlCheckInStatusTooltipText)
-			.Visibility_Static(&FSourceControlMenuHelpers::GetSourceControlCheckInStatusVisibility)
-			.IsEnabled_Lambda([]() { return CanSourceControlCheckIn(); })
-			[
-				SNew(SHorizontalBox)
-				+ SHorizontalBox::Slot()
-				.AutoWidth()
-				.VAlign(VAlign_Center)
-				.HAlign(HAlign_Center)
-				[
-					SNew(SImage)
-					.Image_Static(&FSourceControlMenuHelpers::GetSourceControlCheckInStatusIcon)
-				]
-				+ SHorizontalBox::Slot()
-				.AutoWidth()
-				.VAlign(VAlign_Center)
-				.Padding(FMargin(5, 0, 0, 0))
-				[
-					SNew(STextBlock)
-					.TextStyle(&FAppStyle::Get().GetWidgetStyle<FTextBlockStyle>("NormalText"))
-					.Text_Static(&FSourceControlMenuHelpers::GetSourceControlCheckInStatusText)
-				]
-			]
-			.OnClicked_Static(&FSourceControlMenuHelpers::OnSourceControlCheckInChangesClicked)
-		]
-		+SHorizontalBox::Slot() // Check In Kebab Combo button
-		.VAlign(VAlign_Center)
-		.AutoWidth()
-		[
-			SNew(SComboButton)
-			.ContentPadding(FMargin(7.f, 0.f))
-			.ComboButtonStyle(&FAppStyle::Get().GetWidgetStyle<FComboButtonStyle>("StatusBar.StatusBarEllipsisComboButton"))
-			.MenuPlacement(MenuPlacement_AboveAnchor)
-			.Visibility_Static(&FSourceControlMenuHelpers::GetSourceControlCheckInStatusVisibility)
-			.OnGetMenuContent(FOnGetContent::CreateStatic(&FSourceControlMenuHelpers::GenerateCheckInComboButtonContent))
-		]
-		+ SHorizontalBox::Slot()
-		.VAlign(VAlign_Center)
-		.AutoWidth()
-		[
-			SNew(SSeparator)
-			.Visibility_Static(&FSourceControlMenuHelpers::GetSourceControlCheckInStatusVisibility)
-			.Thickness(1.0)
-			.Orientation(EOrientation::Orient_Vertical)
-		]
-		+ SHorizontalBox::Slot() // Sync Latest Button
-		.VAlign(VAlign_Center)
-		.AutoWidth()
-		[
-			SNew(SButton)
-			.ButtonStyle(&FAppStyle::Get().GetWidgetStyle<FButtonStyle>("StatusBar.StatusBarButton"))
-			.ToolTipText_Static(&FSourceControlMenuHelpers::GetSourceControlSyncStatusTooltipText)
-			.Visibility_Static(&FSourceControlMenuHelpers::GetSourceControlSyncStatusVisibility)
-			.IsEnabled_Lambda([]() { return CanSourceControlSync(); })
-			[
-				SNew(SHorizontalBox)
-				+ SHorizontalBox::Slot()
-				.AutoWidth()
-				.VAlign(VAlign_Center)
-				.HAlign(HAlign_Center)
-				[
-					SNew(SImage)
-					.Image_Static(&FSourceControlMenuHelpers::GetSourceControlSyncStatusIcon)
-				]
-				+ SHorizontalBox::Slot()
-				.AutoWidth()
-				.VAlign(VAlign_Center)
-				.Padding(FMargin(5, 0, 0, 0))
-				[
-					SNew(STextBlock)
-					.TextStyle(&FAppStyle::Get().GetWidgetStyle<FTextBlockStyle>("NormalText"))
-					.Text_Static(&FSourceControlMenuHelpers::GetSourceControlSyncStatusText)
-				]
-			]
-			.OnClicked_Static(&FSourceControlMenuHelpers::OnSourceControlSyncClicked)
-		]
-		+ SHorizontalBox::Slot()
-		.VAlign(VAlign_Center)
-		.AutoWidth()
-		[
-			SNew(SSeparator)
-			.Visibility_Static(&FSourceControlMenuHelpers::GetSourceControlCheckInStatusVisibility)
-			.Thickness(1.0)
-			.Orientation(EOrientation::Orient_Vertical)
+			SNew(SSourceControlControls)
+			.OnClickedSyncLatest_Static(&FSourceControlMenuHelpers::OnSourceControlSyncClicked)
+			.OnClickedCheckInChanges_Static(&FSourceControlMenuHelpers::OnSourceControlCheckInChangesClicked)
+			.OnGenerateKebabMenu_Static(&FSourceControlMenuHelpers::GenerateCheckInComboButtonContent)
 		]
 		+ SHorizontalBox::Slot() // Source Control Menu
 		.VAlign(VAlign_Center)
diff --git a/Engine/Source/Editor/StatusBar/Private/SourceControlMenuHelpers.h b/Engine/Source/Editor/StatusBar/Private/SourceControlMenuHelpers.h
--- a/Engine/Source/Editor/StatusBar/Private/SourceControlMenuHelpers.h
+++ b/Engine/Source/Editor/StatusBar/Private/SourceControlMenuHelpers.h
@@ -9,6 +9,7 @@
 class FUICommandList;
 class FReply;
 class SWidget;
+class SSourceControlControls;
 
 class FSourceControlCommands : public TCommands<FSourceControlCommands>
 {
@@ -75,37 +76,19 @@ private:
 	static FText GetSourceControlTooltip();
 	static const FSlateBrush* GetSourceControlIconBadge();
 
-	/** Sync button */
-	static bool IsAtLatestRevision();
-	static bool CanSourceControlSync();
-	static EVisibility GetSourceControlSyncStatusVisibility();
-	static FText GetSourceControlSyncStatusText();
-	static FText GetSourceControlSyncStatusTooltipText();
-	static const FSlateBrush* GetSourceControlSyncStatusIcon();
-	static FReply OnSourceControlSyncClicked();
-
-	/** Check-in button */
-	static int GetNumLocalChanges();
-	static bool CanSourceControlCheckIn();
 	static EVisibility GetSourceControlCheckInStatusVisibility();
-	static FText GetSourceControlCheckInStatusText();
-	static FText GetSourceControlCheckInStatusTooltipText();
-	static const FSlateBrush* GetSourceControlCheckInStatusIcon();
+
+	static FReply OnSourceControlSyncClicked();
 	static FReply OnSourceControlCheckInChangesClicked();
 
 	/** Callbacks */
 	static void OnSourceControlProviderChanged(ISourceControlProvider& OldProvider, ISourceControlProvider& NewProvider);
 	static void OnSourceControlStateChanged();
 
-	/** Conflicts */
-	static bool AreConflictsRemaining();
-
 private:
 	/** Delegate handles */
-	static FDelegateHandle SourceControlProviderChangedHandle;
-	static FDelegateHandle SourceControlStateChangedHandle;
+	FDelegateHandle SourceControlProviderChangedHandle;
+	FDelegateHandle SourceControlStateChangedHandle;
 
-	/** Is there a conflict remaining? */
-	static bool bConflictsRemaining;
 };
 
diff --git a/Engine/Source/Runtime/SlateCore/Private/Styling/StarshipCoreStyle.cpp b/Engine/Source/Runtime/SlateCore/Private/Styling/StarshipCoreStyle.cpp
--- a/Engine/Source/Runtime/SlateCore/Private/Styling/StarshipCoreStyle.cpp
+++ b/Engine/Source/Runtime/SlateCore/Private/Styling/StarshipCoreStyle.cpp
@@ -435,6 +435,7 @@ TSharedRef<ISlateStyle> FStarshipCoreStyle::Create()
 		Style->Set("Icons.Role", new IMAGE_BRUSH_SVG("Starship/Common/Role", Icon16x16));
 		Style->Set("Icons.Merge", new IMAGE_BRUSH_SVG("Starship/Common/Merge", Icon16x16));
 
+		Style->Set("Icons.Calendar", new IMAGE_BRUSH_SVG("Starship/Common/Calendar", Icon16x16));
 	}
 
 	// Tool panels
