Skein / Snapshot History: Add 'View Snapshot History' menu item to Revision Control menu in Status Bar.

#rnx

diff --git a/Engine/Source/Editor/SourceControlWindows/Private/SourceControlWindowsModule.cpp b/Engine/Source/Editor/SourceControlWindows/Private/SourceControlWindowsModule.cpp
--- a/Engine/Source/Editor/SourceControlWindows/Private/SourceControlWindowsModule.cpp
+++ b/Engine/Source/Editor/SourceControlWindows/Private/SourceControlWindowsModule.cpp
@@ -37,6 +37,13 @@ public:
 
 	virtual void ShowChangelistsTab() override;
 	virtual bool CanShowChangelistsTab() const override;
+
+	virtual void ShowSnapshotHistoryTab() override;
+	virtual bool CanShowSnapshotHistoryTab() const override;
+
+	virtual void ShowConflictResolutionTab() override;
+	virtual bool CanShowConflictResolutionTab() const override;
+
 	virtual void SelectFiles(const TArray<FString>& Filenames);
 
 	DECLARE_DERIVED_EVENT(FSourceControlWindowsModule, ISourceControlWindowsModule::FChangelistFileDoubleClickedEvent, FChangelistFileDoubleClickedEvent);
@@ -126,6 +133,36 @@ bool FSourceControlWindowsModule::CanShowChangelistsTab() const
 	return (SourceControlModule.IsEnabled() && SourceControlModule.GetProvider().IsAvailable() && SourceControlModule.GetProvider().UsesChangelists()) || (FUncontrolledChangelistsModule::Get().IsEnabled() && ISourceControlModule::Get().GetProvider().UsesUncontrolledChangelists());
 }
 
+void FSourceControlWindowsModule::ShowSnapshotHistoryTab()
+{
+	if (IConsoleObject* CObj = IConsoleManager::Get().FindConsoleObject(TEXT("UnrealRevisionControl.FocusSnapshotHistory")))
+	{
+		CObj->AsCommand()->Execute(/*Args=*/TArray<FString>(), /*InWorld=*/nullptr, *GLog);
+	}
+}
+
+bool FSourceControlWindowsModule::CanShowSnapshotHistoryTab() const
+{
+	ISourceControlModule& SourceControlModule = ISourceControlModule::Get();
+
+	return (SourceControlModule.IsEnabled() && SourceControlModule.GetProvider().IsAvailable() && SourceControlModule.GetProvider().GetName() == TEXT("Unreal Revision Control"));
+}
+
+void FSourceControlWindowsModule::ShowConflictResolutionTab()
+{
+	if (IConsoleObject* CObj = IConsoleManager::Get().FindConsoleObject(TEXT("UnrealRevisionControl.FocusConflictResolution")))
+	{
+		CObj->AsCommand()->Execute(/*Args=*/TArray<FString>(), /*InWorld=*/nullptr, *GLog);
+	}
+}
+
+bool FSourceControlWindowsModule::CanShowConflictResolutionTab() const
+{
+	ISourceControlModule& SourceControlModule = ISourceControlModule::Get();
+
+	return (SourceControlModule.IsEnabled() && SourceControlModule.GetProvider().IsAvailable() && SourceControlModule.GetProvider().GetName() == TEXT("Unreal Revision Control"));
+}
+
 void FSourceControlWindowsModule::SelectFiles(const TArray<FString>& Filenames)
 {
 	if (Filenames.Num() > 0 && CanShowChangelistsTab())
diff --git a/Engine/Source/Editor/SourceControlWindows/Public/ISourceControlWindowsModule.h b/Engine/Source/Editor/SourceControlWindows/Public/ISourceControlWindowsModule.h
--- a/Engine/Source/Editor/SourceControlWindows/Public/ISourceControlWindowsModule.h
+++ b/Engine/Source/Editor/SourceControlWindows/Public/ISourceControlWindowsModule.h
@@ -46,6 +46,13 @@ public:
 
 	virtual void ShowChangelistsTab() = 0;
 	virtual bool CanShowChangelistsTab() const = 0;
+
+	virtual void ShowSnapshotHistoryTab() = 0;
+	virtual bool CanShowSnapshotHistoryTab() const = 0;
+
+	virtual void ShowConflictResolutionTab() = 0;
+	virtual bool CanShowConflictResolutionTab() const = 0;
+
 	virtual void SelectFiles(const TArray<FString>& Filenames) = 0;
 
 	DECLARE_EVENT_OneParam(ISourceControlWindowsModule, FChangelistFileDoubleClickedEvent, const FString&);
diff --git a/Engine/Source/Editor/StatusBar/Private/SourceControlMenuHelpers.cpp b/Engine/Source/Editor/StatusBar/Private/SourceControlMenuHelpers.cpp
--- a/Engine/Source/Editor/StatusBar/Private/SourceControlMenuHelpers.cpp
+++ b/Engine/Source/Editor/StatusBar/Private/SourceControlMenuHelpers.cpp
@@ -51,6 +51,7 @@ void FSourceControlCommands::RegisterCommands()
 	UI_COMMAND(ConnectToSourceControl, "Connect to Revision Control...", "Connect to a revision control system for tracking changes to your content and levels.", EUserInterfaceActionType::Button, FInputChord());
 	UI_COMMAND(ChangeSourceControlSettings, "Change Revision Control Settings...", "Opens a dialog to change revision control settings.", EUserInterfaceActionType::Button, FInputChord());
 	UI_COMMAND(ViewChangelists, "View Changes", "Opens a dialog displaying current changes.", EUserInterfaceActionType::Button, FInputChord());
+	UI_COMMAND(ViewSnapshotHistory, "View Snapshot History", "Opens a dialog to switch between source control revisions.", EUserInterfaceActionType::Button, FInputChord());
 	UI_COMMAND(SubmitContent, "Submit Content", "Opens a dialog with check in options for content and levels.", EUserInterfaceActionType::Button, FInputChord());
 	UI_COMMAND(CheckOutModifiedFiles, "Check Out Modified Files", "Opens a dialog to check out any assets which have been modified.", EUserInterfaceActionType::Button, FInputChord());
 	UI_COMMAND(RevertAll, "Revert All Files", "Opens a dialog to revert any assets which have been modified.", EUserInterfaceActionType::Button, FInputChord());
@@ -70,7 +71,15 @@ void FSourceControlCommands::RegisterCommands()
 		FExecuteAction::CreateStatic(&FSourceControlCommands::ViewChangelists_Clicked),
 		FCanExecuteAction::CreateStatic(&FSourceControlCommands::ViewChangelists_CanExecute),
 		FIsActionChecked::CreateLambda([]() { return false; }),
-		FIsActionButtonVisible::CreateStatic(&FSourceControlCommands::ViewChangelists_IsVisible)		
+		FIsActionButtonVisible::CreateStatic(&FSourceControlCommands::ViewChangelists_IsVisible)
+	);
+
+	ActionList->MapAction(
+		ViewSnapshotHistory,
+		FExecuteAction::CreateStatic(&FSourceControlCommands::ViewSnapshotHistory_Clicked),
+		FCanExecuteAction::CreateStatic(&FSourceControlCommands::ViewSnapshotHistory_CanExecute),
+		FIsActionChecked::CreateLambda([]() { return false; }),
+		FIsActionButtonVisible::CreateStatic(&FSourceControlCommands::ViewSnapshotHistory_IsVisible)
 	);
 
 	ActionList->MapAction(
@@ -111,6 +120,16 @@ bool FSourceControlCommands::ViewChangelists_IsVisible()
 	return ISourceControlModule::Get().GetProvider().UsesChangelists() || ISourceControlModule::Get().GetProvider().UsesUncontrolledChangelists();
 }
 
+bool FSourceControlCommands::ViewSnapshotHistory_CanExecute()
+{
+	return ISourceControlWindowsModule::Get().CanShowSnapshotHistoryTab();
+}
+
+bool FSourceControlCommands::ViewSnapshotHistory_IsVisible()
+{
+	return ISourceControlModule::Get().GetProvider().GetName() == TEXT("Unreal Revision Control");
+}
+
 bool FSourceControlCommands::SubmitContent_IsVisible()
 {
 	if (FSourceControlMenuHelpers::GetSourceControlCheckInStatusVisibility() == EVisibility::Visible)
@@ -131,6 +150,11 @@ void FSourceControlCommands::ViewChangelists_Clicked()
 	ISourceControlWindowsModule::Get().ShowChangelistsTab();
 }
 
+void FSourceControlCommands::ViewSnapshotHistory_Clicked()
+{
+	ISourceControlWindowsModule::Get().ShowSnapshotHistoryTab();
+}
+
 bool FSourceControlCommands::CheckOutModifiedFiles_CanExecute()
 {
 	ISourceControlModule& SourceControlModule = ISourceControlModule::Get();
@@ -229,6 +253,13 @@ TSharedRef<SWidget> FSourceControlMenuHelpers::GenerateSourceControlMenuContent(
 		FSlateIcon(FRevisionControlStyleManager::GetStyleSetName(), "RevisionControl.ChangelistsTab")
 	);
 
+	Section.AddMenuEntry(
+		FSourceControlCommands::Get().ViewSnapshotHistory,
+		TAttribute<FText>(),
+		TAttribute<FText>(),
+		FSlateIcon(FRevisionControlStyleManager::GetStyleSetName(), "RevisionControl.Actions.Rewind")
+	);
+
 	Section.AddMenuEntry(
 		FSourceControlCommands::Get().SubmitContent,
 		TAttribute<FText>(),
diff --git a/Engine/Source/Editor/StatusBar/Private/SourceControlMenuHelpers.h b/Engine/Source/Editor/StatusBar/Private/SourceControlMenuHelpers.h
--- a/Engine/Source/Editor/StatusBar/Private/SourceControlMenuHelpers.h
+++ b/Engine/Source/Editor/StatusBar/Private/SourceControlMenuHelpers.h
@@ -26,8 +26,11 @@ private:
 	static void ConnectToSourceControl_Clicked();
 	static bool ViewChangelists_CanExecute();
 	static bool ViewChangelists_IsVisible();
+	static bool ViewSnapshotHistory_CanExecute();
+	static bool ViewSnapshotHistory_IsVisible();
 	static bool SubmitContent_IsVisible();
 	static void ViewChangelists_Clicked();
+	static void ViewSnapshotHistory_Clicked();
 	static bool CheckOutModifiedFiles_CanExecute();
 	static void CheckOutModifiedFiles_Clicked();
 	static bool RevertAllModifiedFiles_CanExecute();
@@ -40,6 +43,7 @@ public:
 	TSharedPtr< FUICommandInfo > ConnectToSourceControl;
 	TSharedPtr< FUICommandInfo > ChangeSourceControlSettings;
 	TSharedPtr< FUICommandInfo > ViewChangelists;
+	TSharedPtr< FUICommandInfo > ViewSnapshotHistory;
 	TSharedPtr< FUICommandInfo > SubmitContent;
 	TSharedPtr< FUICommandInfo > CheckOutModifiedFiles;
 	TSharedPtr< FUICommandInfo > RevertAll;
@@ -81,10 +85,6 @@ private:
 	static FReply OnSourceControlSyncClicked();
 	static FReply OnSourceControlCheckInChangesClicked();
 
-	/** Callbacks */
-	static void OnSourceControlProviderChanged(ISourceControlProvider& OldProvider, ISourceControlProvider& NewProvider);
-	static void OnSourceControlStateChanged();
-
 private:
 	/** Delegate handles */
 	FDelegateHandle SourceControlProviderChangedHandle;
