Some on-hover tooltips are missing or wrong in the Snapshot History tab

diff --git a/Engine/Source/Editor/StatusBar/Private/SourceControlMenuHelpers.cpp b/Engine/Source/Editor/StatusBar/Private/SourceControlMenuHelpers.cpp
--- a/Engine/Source/Editor/StatusBar/Private/SourceControlMenuHelpers.cpp
+++ b/Engine/Source/Editor/StatusBar/Private/SourceControlMenuHelpers.cpp
@@ -51,7 +51,7 @@ void FSourceControlCommands::RegisterCommands()
 	UI_COMMAND(ConnectToSourceControl, "Connect to Revision Control...", "Connect to a revision control system for tracking changes to your content and levels.", EUserInterfaceActionType::Button, FInputChord());
 	UI_COMMAND(ChangeSourceControlSettings, "Change Revision Control Settings...", "Opens a dialog to change revision control settings.", EUserInterfaceActionType::Button, FInputChord());
 	UI_COMMAND(ViewChangelists, "View Changes", "Opens a dialog displaying current changes.", EUserInterfaceActionType::Button, FInputChord());
-	UI_COMMAND(ViewSnapshotHistory, "View Snapshot History", "Opens a dialog to switch between source control revisions.", EUserInterfaceActionType::Button, FInputChord());
+	UI_COMMAND(ViewSnapshotHistory, "Open Snapshot History", "See all the changes that have been made to this project over time.", EUserInterfaceActionType::Button, FInputChord());
 	UI_COMMAND(SubmitContent, "Submit Content", "Opens a dialog with check in options for content and levels.", EUserInterfaceActionType::Button, FInputChord());
 	UI_COMMAND(CheckOutModifiedFiles, "Check Out Modified Files", "Opens a dialog to check out any assets which have been modified.", EUserInterfaceActionType::Button, FInputChord());
 	UI_COMMAND(RevertAll, "Revert All Files", "Opens a dialog to revert any assets which have been modified.", EUserInterfaceActionType::Button, FInputChord());
