StatusBar Notifications: Disable fade-in/out slate throttling for progress status notification

Any world tick process was getting impacted by this non-interactive throttling (BiomeCore 1second added walltime on graph regen)

#rnx

diff --git a/Engine/Source/Editor/StatusBar/Private/SStatusBar.cpp b/Engine/Source/Editor/StatusBar/Private/SStatusBar.cpp
--- a/Engine/Source/Editor/StatusBar/Private/SStatusBar.cpp
+++ b/Engine/Source/Editor/StatusBar/Private/SStatusBar.cpp
@@ -946,7 +946,7 @@ void SStatusBar::UpdateProgressStatus()
 				ProgressNotification.ContentWidget = ActiveNotificationProgressWidgetPin.ToSharedRef();
 				ProgressNotification.FadeOutDuration = StatusBarNotificationConstants::NotificationFadeDuration;
 				ProgressNotification.ForWindow = FSlateApplication::Get().FindWidgetWindow(AsShared());
-
+				ProgressNotification.bAllowThrottleWhenFrameRateIsLow = false;
 				ActiveProgressNotificationPin = FSlateNotificationManager::Get().AddNotification(ProgressNotification);
 
 				ActiveProgressNotification = ActiveProgressNotificationPin;
