Fixed using the "open content drawer" keyboard shortcut to dimiss the CB drawer:
- Expanded the API with explicit Open/Toggle/Dismiss CB drawer functions.
- Hooked the global editor command to toggle the CB drawer as users are used to.
- Requested a dismiss when using the shortcut from the CB search field.

diff --git a/Engine/Source/Editor/ContentBrowser/Private/SContentBrowser.cpp b/Engine/Source/Editor/ContentBrowser/Private/SContentBrowser.cpp
--- a/Engine/Source/Editor/ContentBrowser/Private/SContentBrowser.cpp
+++ b/Engine/Source/Editor/ContentBrowser/Private/SContentBrowser.cpp
@@ -2830,10 +2830,9 @@ FReply SContentBrowser::OnSearchKeyDown(const FGeometry& Geometry, const FKeyEve
 	{
 		FReply Reply = FReply::Handled().ClearUserFocus(EFocusCause::SetDirectly);
 		
-		// If we are a drawer, re-attempting to open will toggle. We must be open as we received input.
 		if (bIsDrawer)
 		{
-			GEditor->GetEditorSubsystem<UStatusBarSubsystem>()->OpenContentBrowserDrawer();
+			GEditor->GetEditorSubsystem<UStatusBarSubsystem>()->DismissContentBrowserDrawer();
 		}
 		return Reply;
 	}
diff --git a/Engine/Source/Editor/StatusBar/Private/StatusBarSubsystem.cpp b/Engine/Source/Editor/StatusBar/Private/StatusBarSubsystem.cpp
--- a/Engine/Source/Editor/StatusBar/Private/StatusBarSubsystem.cpp
+++ b/Engine/Source/Editor/StatusBar/Private/StatusBarSubsystem.cpp
@@ -59,6 +59,31 @@ namespace UE
 
 				return ParentWindow;
 			}
+
+			TSharedPtr<SWindow> FindParentWindowBehindNotification()
+			{
+				TSharedPtr<SWindow> ParentWindow = FSlateApplication::Get().GetActiveTopLevelWindow();
+				if (!ParentWindow.IsValid())
+				{
+					if (TSharedPtr<SDockTab> ActiveTab = FGlobalTabmanager::Get()->GetActiveTab())
+					{
+						if (TSharedPtr<FTabManager> ActiveTabManager = ActiveTab->GetTabManagerPtr())
+						{
+							if (TSharedPtr<SDockTab> ActiveMajorTab = FGlobalTabmanager::Get()->GetMajorTabForTabManager(ActiveTabManager.ToSharedRef()))
+							{
+								ParentWindow = ActiveMajorTab->GetParentWindow();
+							}
+						}
+					}
+				}
+				else if (ParentWindow->GetType() == EWindowType::Notification)
+				{
+					// Get the parent window directly behind the notification. 
+					ParentWindow = FSlateApplication::Get().GetActiveTopLevelRegularWindow(); 
+				}
+
+				return ParentWindow;
+			}
 		}
 	}
 }
@@ -329,54 +354,27 @@ bool UStatusBarSubsystem::ToggleDebugConsole(TSharedRef<SWindow> ParentWindow, b
 
 bool UStatusBarSubsystem::OpenContentBrowserDrawer()
 {
-	TSharedPtr<SWindow> ParentWindow = UE::StatusBarSubsystem::Private::FindParentWindow();
-
-	if (ParentWindow.IsValid() && ParentWindow->GetType() == EWindowType::Normal)
-	{
-		bool bDrawerIsAlreadyOpened = false;
-		for (const TPair<FName, FStatusBarData>& StatusBar : StatusBars)
-		{
-			if (TSharedPtr<SStatusBar> StatusBarPinned = StatusBar.Value.StatusBarWidget.Pin())
-			{
-				if (StatusBarPinned->IsDrawerOpened(StatusBarDrawerIds::ContentBrowser))
-				{
-					TSharedPtr<SDockTab> ParentTab = StatusBarPinned->GetParentTab();
-					if (ParentTab && ParentTab->IsForeground() && ParentTab->GetParentWindow() == ParentWindow)
-					{
-						bDrawerIsAlreadyOpened = true;
-						break;
-					}
-				}
-			}
-		}
+	return TriggerContentBrowser(EDrawerTriggerMode::Open);
+}
 
-		if (!bDrawerIsAlreadyOpened)
-		{
-			TSharedRef<SWindow> WindowRef = ParentWindow.ToSharedRef();
-			return ToggleContentBrowser(ParentWindow.ToSharedRef());
-		}
-		else
-		{
-			return true;
-		}
-	}
+bool UStatusBarSubsystem::ToggleContentBrowserDrawer()
+{
+	return TriggerContentBrowser(EDrawerTriggerMode::Toggle);
+}
 
-	return false;
+bool UStatusBarSubsystem::DismissContentBrowserDrawer()
+{
+	return TriggerContentBrowser(EDrawerTriggerMode::Dismiss);
 }
 
 bool UStatusBarSubsystem::OpenOutputLogDrawer()
 {
-	TSharedPtr<SWindow> ParentWindow = UE::StatusBarSubsystem::Private::FindParentWindow();
-
-	if (ParentWindow.IsValid() && ParentWindow->GetType() == EWindowType::Notification)
-	{
-		// Get the parent window directly behind the notification. 
-		ParentWindow = FSlateApplication::Get().GetActiveTopLevelRegularWindow(); 
-	}
+	TSharedPtr<SWindow> ParentWindow = UE::StatusBarSubsystem::Private::FindParentWindowBehindNotification();
 	
 	if (ParentWindow.IsValid() && ParentWindow->GetType() == EWindowType::Normal)
 	{
-		return ToggleDebugConsole(ParentWindow.ToSharedRef(), true);
+		constexpr bool bAlwaysToggleDrawer = true;
+		return ToggleDebugConsole(ParentWindow.ToSharedRef(), bAlwaysToggleDrawer);
 	}
 
 	return false;
@@ -386,7 +384,7 @@ bool UStatusBarSubsystem::TryToggleDrawer(const FName DrawerId)
 {
 	bool bToggledSuccessfully = false;
 
-	TSharedPtr<SWindow> ParentWindow = UE::StatusBarSubsystem::Private::FindParentWindow();
+	TSharedPtr<SWindow> ParentWindow = UE::StatusBarSubsystem::Private::FindParentWindowBehindNotification();
 
 	if (ParentWindow.IsValid() && ParentWindow->GetType() == EWindowType::Normal)
 	{
@@ -431,9 +429,16 @@ bool UStatusBarSubsystem::ForceDismissDrawer()
 	return bWasDismissed;
 }
 
-bool UStatusBarSubsystem::ToggleContentBrowser(TSharedRef<SWindow> ParentWindow)
+bool UStatusBarSubsystem::TriggerContentBrowser(EDrawerTriggerMode DrawerTriggerMode)
 {
-	bool bWasDismissed = false;
+	TSharedPtr<SWindow> ParentWindow = UE::StatusBarSubsystem::Private::FindParentWindowBehindNotification();
+
+	if (!ParentWindow)
+	{
+		return false;
+	}
+
+	bool bWasAlreadyOpened = false;
 
 	SNewUserTipNotification::Dismiss();
 
@@ -446,14 +451,18 @@ bool UStatusBarSubsystem::ToggleContentBrowser(TSharedRef<SWindow> ParentWindow)
 				TSharedPtr<SDockTab> ParentTab = StatusBarPinned->GetParentTab();
 				if (ParentTab && ParentTab->IsForeground() && ParentTab->GetParentWindow() == ParentWindow)
 				{
-					StatusBarPinned->DismissDrawer(nullptr);
-					bWasDismissed = true;
+					if (EnumHasAllFlags(DrawerTriggerMode, EDrawerTriggerMode::Dismiss))
+					{
+						StatusBarPinned->DismissDrawer(nullptr);
+					}
+
+					bWasAlreadyOpened = true;
 				}
 			}
 		}
 	}
 
-	if(!bWasDismissed)
+	if(!bWasAlreadyOpened && EnumHasAllFlags(DrawerTriggerMode, EDrawerTriggerMode::Open))
 	{
 		TSharedPtr<SWindow> Window = ParentWindow;
 		GEditor->GetTimerManager()->SetTimerForNextTick(FTimerDelegate::CreateUObject(this, &UStatusBarSubsystem::HandleDeferredOpenContentBrowser, Window));
diff --git a/Engine/Source/Editor/StatusBar/Public/StatusBarSubsystem.h b/Engine/Source/Editor/StatusBar/Public/StatusBarSubsystem.h
--- a/Engine/Source/Editor/StatusBar/Public/StatusBarSubsystem.h
+++ b/Engine/Source/Editor/StatusBar/Public/StatusBarSubsystem.h
@@ -79,19 +79,33 @@ public:
 	 */
 	bool ToggleDebugConsole(TSharedRef<SWindow> ParentWindow, bool bAlwaysToggleDrawer=false);
 
+	/**
+	 * Opens the output log drawer for a status bar residing in the active window
+	 *
+	 * @return true if the output log was opened
+	 */
+	bool OpenOutputLogDrawer();
+
 	/**
 	 * Opens the content browser drawer for a status bar residing in the active window 
 	 * 
-	 * @return true if the content browser was opened or false if no status bar in the active window was found
+	 * @return true if the content browser was opened
 	 */
 	bool OpenContentBrowserDrawer();
 
 	/**
-	 * Opens the output log drawer for a status bar residing in the active window
-	 *
-	 * @return true if the output log was opened or false if no status bar in the active window was found
+	 * Opens or closes the content browser drawer for a status bar residing in the active window 
+	 * 
+	 * @return true if the content browser was toggled or false if no status bar in the active window was found
 	 */
-	bool OpenOutputLogDrawer();
+	bool ToggleContentBrowserDrawer();
+
+	/**
+	 * Closes the content browser drawer for a status bar residing in the active window 
+	 * 
+	 * @return true if the content browser was dismissed
+	 */
+	bool DismissContentBrowserDrawer();
 
 	/**
 	 * Tries to toggle the given drawer
@@ -187,7 +201,19 @@ private:
 	virtual void UpdateProgressNotification(FProgressNotificationHandle Handle, int32 TotalWorkDone, int32 UpdatedTotalWorkToDo, FText UpdatedDisplayText) override;
 	virtual void CancelProgressNotification(FProgressNotificationHandle Handle) override;
 
-	bool ToggleContentBrowser(TSharedRef<SWindow> ParentWindow);
+	enum class EDrawerTriggerMode : uint8
+	{
+		None = 0,
+
+		Open = 1 << 0,
+		Dismiss = 1 << 1,
+
+		Toggle = Open | Dismiss
+	};
+	FRIEND_ENUM_CLASS_FLAGS(EDrawerTriggerMode);
+
+	bool TriggerContentBrowser(EDrawerTriggerMode DrawerTriggerMode);
+
 	void OnDebugConsoleClosed(TWeakPtr<SStatusBar> OwningStatusBar);
 	void CreateContentBrowserIfNeeded();
 	void CreateAndShowNewUserTipIfNeeded(TSharedPtr<SWindow> ParentWindow, bool bIsRunningStartupDialog);
diff --git a/Engine/Source/Editor/UnrealEd/Private/Toolkits/GlobalEditorCommonCommands.cpp b/Engine/Source/Editor/UnrealEd/Private/Toolkits/GlobalEditorCommonCommands.cpp
--- a/Engine/Source/Editor/UnrealEd/Private/Toolkits/GlobalEditorCommonCommands.cpp
+++ b/Engine/Source/Editor/UnrealEd/Private/Toolkits/GlobalEditorCommonCommands.cpp
@@ -228,7 +228,7 @@ void FGlobalEditorCommonCommands::OnSummonedConsoleCommandBox()
 
 void FGlobalEditorCommonCommands::OnOpenContentBrowserDrawer()
 {
-	GEditor->GetEditorSubsystem<UStatusBarSubsystem>()->OpenContentBrowserDrawer();
+	GEditor->GetEditorSubsystem<UStatusBarSubsystem>()->ToggleContentBrowserDrawer();
 }
 
 void FGlobalEditorCommonCommands::OnOpenOutputLogDrawer()
