InteractiveToolsFramework: Add tooltip for the status bar at the bottom of the edtior.

diff --git a/Engine/Source/Editor/StatusBar/Private/SStatusBar.cpp b/Engine/Source/Editor/StatusBar/Private/SStatusBar.cpp
--- a/Engine/Source/Editor/StatusBar/Private/SStatusBar.cpp
+++ b/Engine/Source/Editor/StatusBar/Private/SStatusBar.cpp
@@ -766,6 +766,11 @@ TSharedRef<SWidget> SStatusBar::MakeStatusMessageWidget()
 {
 	return 
 		SNew(SHorizontalBox)
+		// Allow the user to mouse over the status bar to read all of it in case it's too long
+		.ToolTipText(this, &SStatusBar::GetStatusBarMessage)
+		// Default visibility is SelfHitTestInvisible. We want it to be hittestable so we can
+		//  show the tooltip when mousing over any part of the bar.
+		.Visibility(EVisibility::Visible)
 		+ SHorizontalBox::Slot()
 		.AutoWidth()
 		.VAlign(VAlign_Center)
