Notifications: Added an optional copy to clipboard button in the top right corner
- Opted in to the copy to clipboard button for changelist submitted notifications
- Deprecated the ShowCopyToClipboadHyperlink function as it prevented having both an hyperlink and a mean to copy to the clipboard
- Replaced usages of ShowCopyToClipboadHyperlink with bUseCopyToClipboard

diff --git a/Engine/Source/Editor/SourceControlWindows/Private/SSourceControlChangelists.cpp b/Engine/Source/Editor/SourceControlWindows/Private/SSourceControlChangelists.cpp
--- a/Engine/Source/Editor/SourceControlWindows/Private/SSourceControlChangelists.cpp
+++ b/Engine/Source/Editor/SourceControlWindows/Private/SSourceControlChangelists.cpp
@@ -2485,7 +2485,13 @@ void SSourceControlChangelistsWidget::OnSubmitChangelist()
 				{
 					if (InResult == ECommandResult::Succeeded)
 					{
-						SSourceControlCommon::DisplaySourceControlOperationNotification(SubmitChangelistOperation->GetSuccessMessage(), SNotificationItem::CS_Success);
+						TOptional<FNotificationInfo> NotificationInfo = SSourceControlCommon::ConstructSourceControlOperationNotification(SubmitChangelistOperation->GetSuccessMessage());
+						if (NotificationInfo)
+						{
+							NotificationInfo->bUseCopyToClipboard = true;
+							SSourceControlCommon::DisplaySourceControlOperationNotification(NotificationInfo.GetValue(), SNotificationItem::CS_Success);
+						}
+
 						bCheckinSuccess = true;
 					}
 					else if (InResult == ECommandResult::Failed)
diff --git a/Engine/Source/Editor/SourceControlWindows/Private/SSourceControlCommon.cpp b/Engine/Source/Editor/SourceControlWindows/Private/SSourceControlCommon.cpp
--- a/Engine/Source/Editor/SourceControlWindows/Private/SSourceControlCommon.cpp
+++ b/Engine/Source/Editor/SourceControlWindows/Private/SSourceControlCommon.cpp
@@ -515,20 +515,39 @@ void ExecuteUncontrolledChangelistOperationWithSlowTaskWrapper(const FText& Mess
 	ExecuteChangelistOperationWithSlowTaskWrapper(Message, UncontrolledChangelistTask);
 }
 
-/** Displays toast notification to report the status of task. */
-void DisplaySourceControlOperationNotification(const FText& Message, SNotificationItem::ECompletionState CompletionState)
+TOptional<FNotificationInfo> ConstructSourceControlOperationNotification(const FText& Message)
 {
 	if (Message.IsEmpty())
 	{
-		return;
+		return {};
 	}
 
-	FMessageLog("SourceControl").Message(CompletionState == SNotificationItem::ECompletionState::CS_Fail ? EMessageSeverity::Error : EMessageSeverity::Info, Message);
-
 	FNotificationInfo NotificationInfo(Message);
 	NotificationInfo.ExpireDuration = 6.0f;
 	NotificationInfo.Hyperlink = FSimpleDelegate::CreateLambda([]() { FGlobalTabmanager::Get()->TryInvokeTab(FName("OutputLog")); });
 	NotificationInfo.HyperlinkText = LOCTEXT("ShowOutputLogHyperlink", "Show Output Log");
+	
+	return NotificationInfo;
+}
+
+/** Displays toast notification to report the status of task. */
+void DisplaySourceControlOperationNotification(const FText& Message, SNotificationItem::ECompletionState CompletionState)
+{
+	if (TOptional<FNotificationInfo> NotificationInfo = ConstructSourceControlOperationNotification(Message))
+	{
+		DisplaySourceControlOperationNotification(NotificationInfo.GetValue(), CompletionState);
+	}
+}
+
+void DisplaySourceControlOperationNotification(const FNotificationInfo& NotificationInfo, SNotificationItem::ECompletionState CompletionState)
+{
+	if (!NotificationInfo.Text.IsSet())
+	{
+		return;
+	}
+
+	FMessageLog("SourceControl").Message(CompletionState == SNotificationItem::ECompletionState::CS_Fail ? EMessageSeverity::Error : EMessageSeverity::Info, NotificationInfo.Text.Get());
+
 	FSlateNotificationManager::Get().AddNotification(NotificationInfo)->SetCompletionState(CompletionState);
 }
 
diff --git a/Engine/Source/Editor/SourceControlWindows/Private/SSourceControlCommon.h b/Engine/Source/Editor/SourceControlWindows/Private/SSourceControlCommon.h
--- a/Engine/Source/Editor/SourceControlWindows/Private/SSourceControlCommon.h
+++ b/Engine/Source/Editor/SourceControlWindows/Private/SSourceControlCommon.h
@@ -410,7 +410,11 @@ namespace SSourceControlCommon
 
 	void ExecuteChangelistOperationWithSlowTaskWrapper(const FText& Message, const TFunction<void()>& ChangelistTask);
 	void ExecuteUncontrolledChangelistOperationWithSlowTaskWrapper(const FText& Message, const TFunction<void()>& UncontrolledChangelistTask);
+	
 	void DisplaySourceControlOperationNotification(const FText& Message, SNotificationItem::ECompletionState CompletionState);
+	void DisplaySourceControlOperationNotification(const FNotificationInfo& NotificationInfo, SNotificationItem::ECompletionState CompletionState);
+	TOptional<FNotificationInfo> ConstructSourceControlOperationNotification(const FText& Message);
+
 	bool OpenConflictDialog(const TArray<FSourceControlStateRef>& InFilesConflicts);
 }
 
diff --git a/Engine/Source/Editor/StatusBar/Private/NotificationTesting.cpp b/Engine/Source/Editor/StatusBar/Private/NotificationTesting.cpp
--- a/Engine/Source/Editor/StatusBar/Private/NotificationTesting.cpp
+++ b/Engine/Source/Editor/StatusBar/Private/NotificationTesting.cpp
@@ -68,7 +68,10 @@ static void TestNotifications()
 {
 	static const float Timeout = 15.0f;
 
-	FTSTicker::GetCoreTicker().AddTicker(TEXT("TestNotifications"), 0.0f, [](float DeltaTime)
+	int Delay = 0;
+
+	// Text only
+	FTSTicker::GetCoreTicker().AddTicker(TEXT("TestNotifications"), (float)Delay++, [](float DeltaTime)
 		{
 			FNotificationInfo NotificationInfo(FText::FromString(TEXT("Test Notification 1")));
 			NotificationInfo.FadeInDuration = 2.0f;
@@ -81,13 +84,30 @@ static void TestNotifications()
 			return false;
 		});
 
-	FTSTicker::GetCoreTicker().AddTicker(TEXT("TestNotifications"), 2.0f, [](float DeltaTime)
+	// SubText 
+	FTSTicker::GetCoreTicker().AddTicker(TEXT("TestNotifications"), (float)Delay++, [](float DeltaTime)
 		{
 			FNotificationInfo NotificationInfo(FText::FromString(TEXT("Test Notification 2")));
 			NotificationInfo.FadeInDuration = 2.0f;
 			NotificationInfo.FadeOutDuration = 2.0f;
 			NotificationInfo.ExpireDuration = Timeout;
 
+			NotificationInfo.SubText = FText::FromString(TEXT("SubText Test"));
+
+			auto Notification = FSlateNotificationManager::Get().AddNotification(NotificationInfo);
+			Notification->ExpireAndFadeout();
+
+			return false;
+		});
+
+	// Throbber
+	FTSTicker::GetCoreTicker().AddTicker(TEXT("TestNotifications"), (float)Delay++, [](float DeltaTime)
+		{
+			FNotificationInfo NotificationInfo(FText::FromString(TEXT("Test Notification 3")));
+			NotificationInfo.FadeInDuration = 2.0f;
+			NotificationInfo.FadeOutDuration = 2.0f;
+			NotificationInfo.ExpireDuration = Timeout;
+
 			NotificationInfo.bUseThrobber = true;
 			auto Notification = FSlateNotificationManager::Get().AddNotification(NotificationInfo);
 			Notification->SetCompletionState(SNotificationItem::CS_Pending);
@@ -96,9 +116,10 @@ static void TestNotifications()
 			return false;
 		});
 
-	FTSTicker::GetCoreTicker().AddTicker(TEXT("TestNotifications"), 3.0f, [](float DeltaTime)
+	// Success icon
+	FTSTicker::GetCoreTicker().AddTicker(TEXT("TestNotifications"), (float)Delay++, [](float DeltaTime)
 		{
-			FNotificationInfo NotificationInfo(FText::FromString(TEXT("Test Notification 3")));
+			FNotificationInfo NotificationInfo(FText::FromString(TEXT("Test Notification 4")));
 			NotificationInfo.FadeInDuration = 2.0f;
 			NotificationInfo.FadeOutDuration = 2.0f;
 			NotificationInfo.ExpireDuration = Timeout;
@@ -111,9 +132,10 @@ static void TestNotifications()
 			return false;
 		});
 
-	FTSTicker::GetCoreTicker().AddTicker(TEXT("TestNotifications"), 4.0f, [](float DeltaTime)
+	// Failure icon
+	FTSTicker::GetCoreTicker().AddTicker(TEXT("TestNotifications"), (float)Delay++, [](float DeltaTime)
 		{
-			FNotificationInfo NotificationInfo(FText::FromString(TEXT("Test Notification 4")));
+			FNotificationInfo NotificationInfo(FText::FromString(TEXT("Test Notification 5")));
 			NotificationInfo.FadeInDuration = 2.0f;
 			NotificationInfo.FadeOutDuration = 2.0f;
 			NotificationInfo.ExpireDuration = Timeout;
@@ -126,9 +148,10 @@ static void TestNotifications()
 			return false;
 		});
 
-	FTSTicker::GetCoreTicker().AddTicker(TEXT("TestNotifications"), 5.0f, [](float DeltaTime)
+	// Checkbox
+	FTSTicker::GetCoreTicker().AddTicker(TEXT("TestNotifications"), (float)Delay++, [](float DeltaTime)
 		{
-			FNotificationInfo NotificationInfo(FText::FromString(TEXT("Test Notification 5")));
+			FNotificationInfo NotificationInfo(FText::FromString(TEXT("Test Notification 6")));
 			NotificationInfo.FadeInDuration = 2.0f;
 			NotificationInfo.FadeOutDuration = 2.0f;
 			NotificationInfo.ExpireDuration = Timeout;
@@ -144,9 +167,10 @@ static void TestNotifications()
 			return false;
 		});
 
-	FTSTicker::GetCoreTicker().AddTicker(TEXT("TestNotifications"), 6.0f, [](float DeltaTime)
+	// Hyperlink
+	FTSTicker::GetCoreTicker().AddTicker(TEXT("TestNotifications"), (float)Delay++, [](float DeltaTime)
 		{
-			FNotificationInfo NotificationInfo(FText::FromString(TEXT("Test Notification 6")));
+			FNotificationInfo NotificationInfo(FText::FromString(TEXT("Test Notification 7")));
 			NotificationInfo.FadeInDuration = 2.0f;
 			NotificationInfo.FadeOutDuration = 2.0f;
 			NotificationInfo.ExpireDuration = Timeout;
@@ -161,9 +185,10 @@ static void TestNotifications()
 			return false;
 		});
 
-	FTSTicker::GetCoreTicker().AddTicker(TEXT("TestNotifications"), 7.0f, [](float DeltaTime)
+	// Buttons
+	FTSTicker::GetCoreTicker().AddTicker(TEXT("TestNotifications"), (float)Delay++, [](float DeltaTime)
 		{
-			FNotificationInfo NotificationInfo(FText::FromString(TEXT("Test Notification 7")));
+			FNotificationInfo NotificationInfo(FText::FromString(TEXT("Test Notification 8")));
 			NotificationInfo.FadeInDuration = 2.0f;
 			NotificationInfo.FadeOutDuration = 2.0f;
 			NotificationInfo.ExpireDuration = Timeout;
@@ -180,13 +205,32 @@ static void TestNotifications()
 			return false;
 		});
 
-	FTSTicker::GetCoreTicker().AddTicker(TEXT("TestNotifications"), 8.0f, [](float DeltaTime)
+	// Copy to Clipboard
+	FTSTicker::GetCoreTicker().AddTicker(TEXT("TestNotifications"), (float)Delay++, [](float DeltaTime)
+		{
+			FNotificationInfo NotificationInfo(FText::FromString(TEXT("Test Notification 9")));
+			NotificationInfo.FadeInDuration = 2.0f;
+			NotificationInfo.FadeOutDuration = 2.0f;
+			NotificationInfo.ExpireDuration = Timeout;
+
+			NotificationInfo.bUseCopyToClipboard = true;
+
+			auto Notification = FSlateNotificationManager::Get().AddNotification(NotificationInfo);
+			Notification->ExpireAndFadeout();
+
+			return false;
+		});
+
+	// Everything
+	FTSTicker::GetCoreTicker().AddTicker(TEXT("TestNotifications"), (float)Delay++, [](float DeltaTime)
 		{
 			FNotificationInfo NotificationInfo(FText::FromString(TEXT("Everything Under The Sun. This one also has a lot of text which should wrap to the next line")));
 			NotificationInfo.FadeInDuration = 2.0f;
 			NotificationInfo.FadeOutDuration = 2.0f;
 			NotificationInfo.ExpireDuration = Timeout;
 
+			NotificationInfo.SubText = FText::FromString(TEXT("SubText Test"));
+
 			NotificationInfo.CheckBoxText = FText::FromString(TEXT("Don't ask again"));
 			NotificationInfo.CheckBoxState = ECheckBoxState::Checked;
 			NotificationInfo.CheckBoxStateChanged = FOnCheckStateChanged::CreateStatic([](ECheckBoxState NewState) {});
@@ -196,7 +240,7 @@ static void TestNotifications()
 
 			NotificationInfo.bUseSuccessFailIcons = true;
 			NotificationInfo.bUseThrobber = true;
-
+			NotificationInfo.bUseCopyToClipboard = true;
 
 			FNotificationButtonInfo Button1(FText::FromString("OK"), FText::GetEmpty(), FSimpleDelegate());
 			FNotificationButtonInfo Button2(FText::FromString("CANCEL"), FText::GetEmpty(), FSimpleDelegate());
@@ -213,13 +257,16 @@ static void TestNotifications()
 			return false;
 		});
 
-	FTSTicker::GetCoreTicker().AddTicker(TEXT("TestNotifications"), 9.0f, [](float DeltaTime)
+	// Everything overflowing
+	FTSTicker::GetCoreTicker().AddTicker(TEXT("TestNotifications"), (float)Delay++, [](float DeltaTime)
 		{
 			FNotificationInfo NotificationInfo(FText::FromString(TEXT("This one has a lot of text on the buttons which should overflow properly")));
 			NotificationInfo.FadeInDuration = 2.0f;
 			NotificationInfo.FadeOutDuration = 2.0f;
 			NotificationInfo.ExpireDuration = Timeout;
 
+			NotificationInfo.SubText = FText::FromString(TEXT("This one also has a lot of subtext which should wrap to the next line"));
+
 			NotificationInfo.CheckBoxText = FText::FromString(TEXT("This is a checkbox with a lot of text. Hover over it to read the full text in the tooltip."));
 			NotificationInfo.CheckBoxState = ECheckBoxState::Checked;
 			NotificationInfo.CheckBoxStateChanged = FOnCheckStateChanged::CreateStatic([](ECheckBoxState NewState) {});
@@ -229,7 +276,7 @@ static void TestNotifications()
 
 			NotificationInfo.bUseSuccessFailIcons = true;
 			NotificationInfo.bUseThrobber = true;
-
+			NotificationInfo.bUseCopyToClipboard = true;
 
 			FNotificationButtonInfo Button1(FText::FromString("This is a button with a lot of text. Hover over it to read the full text in the tooltip."), FText::GetEmpty(), FSimpleDelegate());
 			FNotificationButtonInfo Button2(FText::FromString("This is another button with a lot of text. Hover over it to read the full text in the tooltip."), FText::GetEmpty(), FSimpleDelegate());
diff --git a/Engine/Source/Runtime/Slate/Private/Widgets/Notifications/SNotificationList.cpp b/Engine/Source/Runtime/Slate/Private/Widgets/Notifications/SNotificationList.cpp
--- a/Engine/Source/Runtime/Slate/Private/Widgets/Notifications/SNotificationList.cpp
+++ b/Engine/Source/Runtime/Slate/Private/Widgets/Notifications/SNotificationList.cpp
@@ -393,6 +393,8 @@ public:
 		SLATE_ARGUMENT(bool, bUseSuccessFailIcons)
 		/** When true the larger bolder font will be used to display the message */
 		SLATE_ARGUMENT(bool, bUseLargeFont)
+		/** Controls whether or not to display the copy to clipboard button */
+		SLATE_ARGUMENT(bool, bUseCopyToClipboard)
 		/** When set this forces the width of the box, used to stop resizeing on text change */
 		SLATE_ARGUMENT(FOptionalSize, WidthOverride)
 		/** When set this will display a check box on the notification; handles getting the current check box state */
@@ -448,13 +450,47 @@ public:
 
 		// Container for the text and optional interactive widgets (buttons, check box, and hyperlink)
 		TSharedRef<SVerticalBox> InteractiveWidgetsBox = SNew(SVerticalBox);
+		TSharedRef<SWidget> CopyToClipboardButton = SNullWidget::NullWidget;
+
+		if (InArgs._bUseCopyToClipboard)
+		{
+			SAssignNew(CopyToClipboardButton, SBorder)
+			.HAlign(HAlign_Center)
+			.VAlign(VAlign_Top)
+			.Padding(0.f)
+			.BorderImage(FAppStyle::GetBrush("NoBorder"))
+			[
+				SNew(SButton)
+				.ButtonStyle(FAppStyle::Get(), "SimpleButton")
+				.ToolTipText(LOCTEXT("CopyToClipboard", "Copy to Clipboard"))
+				.OnClicked(this, &SNotificationItemImpl::OnCopyTextToClipboard)
+				.ContentPadding(0.f)
+				[
+					SNew(SImage)
+					.Image(FAppStyle::Get().GetBrush("Icons.Clipboard"))
+					.ColorAndOpacity(FSlateColor::UseForeground())
+				]
+			];
+		}
 
 		InteractiveWidgetsBox->AddSlot()
 		[
-			SAssignNew(MyTextBlock, STextBlock)
-			.Text(Text)
-			.Font(FAppStyle::Get().GetFontStyle(TEXT("NotificationList.FontBold")))
-			.WrapTextAt(InArgs._WidthOverride.IsSet() ? InArgs._WidthOverride.Get()-50.0f : 0.0f )
+			SNew(SHorizontalBox)
+			+ SHorizontalBox::Slot()
+			.FillWidth(1.0f)
+			[
+				SAssignNew(MyTextBlock, STextBlock)
+				.Text(Text)
+				.Font(FAppStyle::Get().GetFontStyle(TEXT("NotificationList.FontBold")))
+				.WrapTextAt(InArgs._WidthOverride.IsSet() ? InArgs._WidthOverride.Get()-50.0f : 0.0f )
+			]
+			+ SHorizontalBox::Slot()
+			.AutoWidth()
+			.HAlign(HAlign_Right)
+			.Padding(0.f)
+			[
+				CopyToClipboardButton
+			]
 		];
 
 		InteractiveWidgetsBox->AddSlot()
@@ -631,6 +667,23 @@ protected:
 		return HyperlinkText.Get();
 	}
 
+	/* Used to fetch the text to copy to the clipboard */
+	FText GetClipboardText() const
+	{
+		FTextBuilder ClipboardTextBuilder;
+		if (Text.IsSet())
+		{
+			ClipboardTextBuilder.AppendLine(Text.Get());
+		}
+
+		if (SubText.IsSet())
+		{
+			ClipboardTextBuilder.AppendLine(SubText.Get());
+		}
+
+		return ClipboardTextBuilder.ToText();
+	}
+
 private:
 
 	/* Used as a wrapper for the callback, this means any code calling it does not require access to FReply type */
@@ -646,6 +699,12 @@ private:
 		Hyperlink.ExecuteIfBound();
 	}
 
+	FReply OnCopyTextToClipboard() const
+	{
+		FPlatformApplicationMisc::ClipboardCopy(*GetClipboardText().ToString());
+		return FReply::Handled();
+	}
+
 protected:
 
 	/** When set this will display a check box on the notification; handles setting the new check box state */
@@ -825,6 +884,7 @@ TSharedRef<SNotificationItem> SNotificationList::AddNotification(const FNotifica
 				.bUseThrobber(Info.bUseThrobber)
 				.bUseSuccessFailIcons(Info.bUseSuccessFailIcons)
 				.bUseLargeFont(Info.bUseLargeFont)
+				.bUseCopyToClipboard(Info.bUseCopyToClipboard)
 				.WidthOverride(Info.WidthOverride)
 				.CheckBoxState(Info.CheckBoxState)
 				.CheckBoxStateChanged(Info.CheckBoxStateChanged)
diff --git a/Engine/Source/Runtime/Slate/Public/Widgets/Notifications/SNotificationList.h b/Engine/Source/Runtime/Slate/Public/Widgets/Notifications/SNotificationList.h
--- a/Engine/Source/Runtime/Slate/Public/Widgets/Notifications/SNotificationList.h
+++ b/Engine/Source/Runtime/Slate/Public/Widgets/Notifications/SNotificationList.h
@@ -170,6 +170,7 @@ struct FNotificationInfo
 		bUseThrobber(false),
 		bUseSuccessFailIcons(false),
 		bUseLargeFont(false),
+		bUseCopyToClipboard(false),
 		WidthOverride(320.0f),
 		bFireAndForget(true),
 		CheckBoxState(ECheckBoxState::Unchecked),
@@ -184,6 +185,7 @@ struct FNotificationInfo
 	 * Shows a "Copy to Clipboard" hyperlink that when clicked copies the text and sub-text to the operating system clipboard
 	 * @note Text and SubText must be set prior to calling this method for them to be included
 	 */
+	UE_DEPRECATED(5.6, "Use the flag bUseCopyToClipboard instead.")
 	SLATE_API void ShowCopyToClipboadHyperlink();
 
 	/** If set, overrides the entire content of the notification with this widget */
@@ -219,6 +221,9 @@ struct FNotificationInfo
 	/** When true the larger bolder font will be used to display the message */
 	bool bUseLargeFont;
 
+	/** Controls whether or not to display the copy to clipboard button */
+	bool bUseCopyToClipboard;
+
 	/** When set this forces the width of the box, used to stop resizing on text change */
 	FOptionalSize WidthOverride;
 
