StatusBarSubsystem - Add ECreateStatusBarOptions to the MakeStatusBarWidget method in order to allow creating a more customized status bar.

This change is needed to allow the LiveLinkHub application to create a status bar that does not contain a content browser and source control section

diff --git a/Engine/Source/Editor/StatusBar/Private/SStatusBar.cpp b/Engine/Source/Editor/StatusBar/Private/SStatusBar.cpp
--- a/Engine/Source/Editor/StatusBar/Private/SStatusBar.cpp
+++ b/Engine/Source/Editor/StatusBar/Private/SStatusBar.cpp
@@ -573,6 +573,7 @@ SStatusBar::~SStatusBar()
 void SStatusBar::Construct(const FArguments& InArgs, FName InStatusBarName, const TSharedRef<SDockTab> InParentTab)
 {
 	ParentTab = InParentTab;
+	bCreateSourceControlSection = InArgs._CreateSourceControlSection;
 
 	UpArrow = FAppStyle::Get().GetBrush("StatusBar.ContentBrowserUp");
 	DownArrow = FAppStyle::Get().GetBrush("StatusBar.ContentBrowserDown");
@@ -757,7 +758,10 @@ TSharedRef<SWidget> SStatusBar::MakeStatusBarToolBarWidget()
 	RegisterStatusBarMenu();
 	
 	FToolMenuContext MenuContext;
-	RegisterSourceControlStatus();
+	if (bCreateSourceControlSection)
+	{
+		RegisterSourceControlStatus();
+	}
 
 	return UToolMenus::Get()->GenerateWidget(GetToolbarName(), MenuContext);
 }
diff --git a/Engine/Source/Editor/StatusBar/Private/SStatusBar.h b/Engine/Source/Editor/StatusBar/Private/SStatusBar.h
--- a/Engine/Source/Editor/StatusBar/Private/SStatusBar.h
+++ b/Engine/Source/Editor/StatusBar/Private/SStatusBar.h
@@ -60,6 +60,7 @@ class SStatusBar : public SCompoundWidget
 {
 	SLATE_BEGIN_ARGS(SStatusBar)	
 	{}
+		SLATE_ARGUMENT_DEFAULT(bool, CreateSourceControlSection) = true;
 	SLATE_END_ARGS()
 
 public:
@@ -200,5 +201,8 @@ private:
 	const FSlateBrush* DownArrow = nullptr;
 
 	bool bAllowedToRefreshProgressNotification = false;
+
+	/** Whether or not to create the source control section. */
+	bool bCreateSourceControlSection = true;
 };
 
diff --git a/Engine/Source/Editor/StatusBar/Private/StatusBarSubsystem.cpp b/Engine/Source/Editor/StatusBar/Private/StatusBarSubsystem.cpp
--- a/Engine/Source/Editor/StatusBar/Private/StatusBarSubsystem.cpp
+++ b/Engine/Source/Editor/StatusBar/Private/StatusBarSubsystem.cpp
@@ -471,86 +471,97 @@ bool UStatusBarSubsystem::TriggerContentBrowser(EDrawerTriggerMode DrawerTrigger
 	return true;
 }
 
-TSharedRef<SWidget> UStatusBarSubsystem::MakeStatusBarWidget(FName StatusBarName, const TSharedRef<SDockTab>& InParentTab)
+TSharedRef<SWidget> UStatusBarSubsystem::MakeStatusBarWidget(FName StatusBarName, const TSharedRef<SDockTab>& InParentTab, UE::Editor::Toolbars::ECreateStatusBarOptions CreateStatusBarOptions)
 {
+	using namespace UE::Editor::Toolbars;
+
 	LLM_SCOPE(ELLMTag::UI);
 
 	CreateContentBrowserIfNeeded();
 
 	TSharedRef<SStatusBar> StatusBar =
-		SNew(SStatusBar, StatusBarName, InParentTab);
-
-	FWidgetDrawerConfig ContentBrowserDrawer(StatusBarDrawerIds::ContentBrowser);
-	ContentBrowserDrawer.GetDrawerContentDelegate.BindUObject(this, &UStatusBarSubsystem::OnGetContentBrowser);
-	ContentBrowserDrawer.OnDrawerOpenedDelegate.BindUObject(this, &UStatusBarSubsystem::OnContentBrowserOpened);
-	ContentBrowserDrawer.OnDrawerDismissedDelegate.BindUObject(this, &UStatusBarSubsystem::OnContentBrowserDismissed);
-	ContentBrowserDrawer.ButtonText = LOCTEXT("StatusBar_ContentBrowserButton", "Content Drawer");
-	ContentBrowserDrawer.ToolTipText = FText::Format(LOCTEXT("StatusBar_ContentBrowserDrawerToolTip", "Opens a temporary content browser above this status which will dismiss when it loses focus ({0})"), FGlobalEditorCommonCommands::Get().OpenContentBrowserDrawer->GetInputText());
-	ContentBrowserDrawer.Icon = FAppStyle::Get().GetBrush("ContentBrowser.TabIcon");
+		SNew(SStatusBar, StatusBarName, InParentTab)
+		.CreateSourceControlSection(!EnumHasAnyFlags(CreateStatusBarOptions, ECreateStatusBarOptions::HideSourceControl));
 
-	for (const TUniquePtr<IGlobalStatusBarExtension>& Extension : GlobalStatusBarExtensions)
+	if (!EnumHasAnyFlags(CreateStatusBarOptions, UE::Editor::Toolbars::ECreateStatusBarOptions::HideContentBrowser))
 	{
-		Extension->ExtendContentBrowserDrawer(ContentBrowserDrawer);
-	}
+		FWidgetDrawerConfig ContentBrowserDrawer(StatusBarDrawerIds::ContentBrowser);
+		ContentBrowserDrawer.GetDrawerContentDelegate.BindUObject(this, &UStatusBarSubsystem::OnGetContentBrowser);
+		ContentBrowserDrawer.OnDrawerOpenedDelegate.BindUObject(this, &UStatusBarSubsystem::OnContentBrowserOpened);
+		ContentBrowserDrawer.OnDrawerDismissedDelegate.BindUObject(this, &UStatusBarSubsystem::OnContentBrowserDismissed);
+		ContentBrowserDrawer.ButtonText = LOCTEXT("StatusBar_ContentBrowserButton", "Content Drawer");
+		ContentBrowserDrawer.ToolTipText = FText::Format(LOCTEXT("StatusBar_ContentBrowserDrawerToolTip", "Opens a temporary content browser above this status which will dismiss when it loses focus ({0})"), FGlobalEditorCommonCommands::Get().OpenContentBrowserDrawer->GetInputText());
+		ContentBrowserDrawer.Icon = FAppStyle::Get().GetBrush("ContentBrowser.TabIcon");
 
-	StatusBar->RegisterDrawer(MoveTemp(ContentBrowserDrawer));
+		for (const TUniquePtr<IGlobalStatusBarExtension>& Extension : GlobalStatusBarExtensions)
+		{
+			Extension->ExtendContentBrowserDrawer(ContentBrowserDrawer);
+		}
+
+		StatusBar->RegisterDrawer(MoveTemp(ContentBrowserDrawer));
+	}
 
 	FOutputLogModule& OutputLogModule = FModuleManager::Get().LoadModuleChecked<FOutputLogModule>("OutputLog");
 
 	TWeakPtr<SStatusBar> StatusBarWeakPtr = StatusBar;
 
-	FSimpleDelegate OnConsoleClosed = FSimpleDelegate::CreateUObject(this, &UStatusBarSubsystem::OnDebugConsoleClosed, StatusBarWeakPtr);
-	FSimpleDelegate OnConsoleCommandExecuted;
-
 	TSharedPtr<SMultiLineEditableTextBox> ConsoleEditBox;
-	TSharedPtr<SWidget> OutputLog;
+	if (!EnumHasAnyFlags(CreateStatusBarOptions, UE::Editor::Toolbars::ECreateStatusBarOptions::HideOutputLog))
 	{
-		TSharedRef<SWidget> ConsoleInputBox = OutputLogModule.MakeConsoleInputBox(ConsoleEditBox, OnConsoleClosed, OnConsoleCommandExecuted);
+		FSimpleDelegate OnConsoleClosed = FSimpleDelegate::CreateUObject(this, &UStatusBarSubsystem::OnDebugConsoleClosed, StatusBarWeakPtr);
+		FSimpleDelegate OnConsoleCommandExecuted;
 
-		auto IsConsoleInputBoxBorderVisible = [ConsoleInputBoxWeak = ConsoleInputBox.ToWeakPtr()]()
+		
+		TSharedPtr<SWidget> OutputLog;
 		{
-			if (TSharedPtr<SWidget> ConsoleInputBox = ConsoleInputBoxWeak.Pin())
-			{
-				FSlateAttributeMetaData::UpdateOnlyVisibilityAttributes(*ConsoleInputBox, FSlateAttributeMetaData::EInvalidationPermission::AllowInvalidationIfConstructed);
-				if (ConsoleInputBox->GetVisibility() == EVisibility::Collapsed)
+			TSharedRef<SWidget> ConsoleInputBox = OutputLogModule.MakeConsoleInputBox(ConsoleEditBox, OnConsoleClosed, OnConsoleCommandExecuted);
+
+			auto IsConsoleInputBoxBorderVisible = [ConsoleInputBoxWeak = ConsoleInputBox.ToWeakPtr()]()
 				{
-					return EVisibility::Collapsed;
-				}
-			}
-			return EVisibility::Visible;
-		};
-
-		OutputLog =
-			SNew(SBorder)
-			.BorderImage(FAppStyle::Get().GetBrush("Brushes.Panel"))
-			.VAlign(VAlign_Center)
-			.Padding(FMargin(6.0f, 0.0f))
-			.Visibility(MakeAttributeLambda(IsConsoleInputBoxBorderVisible))
-			[
-				SNew(SBox)
+					if (TSharedPtr<SWidget> ConsoleInputBox = ConsoleInputBoxWeak.Pin())
+					{
+						FSlateAttributeMetaData::UpdateOnlyVisibilityAttributes(*ConsoleInputBox, FSlateAttributeMetaData::EInvalidationPermission::AllowInvalidationIfConstructed);
+						if (ConsoleInputBox->GetVisibility() == EVisibility::Collapsed)
+						{
+							return EVisibility::Collapsed;
+						}
+					}
+					return EVisibility::Visible;
+				};
+
+			OutputLog =
+				SNew(SBorder)
+				.BorderImage(FAppStyle::Get().GetBrush("Brushes.Panel"))
+				.VAlign(VAlign_Center)
+				.Padding(FMargin(6.0f, 0.0f))
+				.Visibility(MakeAttributeLambda(IsConsoleInputBoxBorderVisible))
 				[
-					ConsoleInputBox
-				]
-			];
-	}
+					SNew(SBox)
+						[
+							ConsoleInputBox
+						]
+				];
+		}
 
-	FWidgetDrawerConfig OutputLogDrawer(StatusBarDrawerIds::OutputLog);
+		FWidgetDrawerConfig OutputLogDrawer(StatusBarDrawerIds::OutputLog);
 
-	OutputLogDrawer.GetDrawerContentDelegate.BindUObject(this, &UStatusBarSubsystem::OnGetOutputLog);
-	OutputLogDrawer.OnDrawerOpenedDelegate.BindUObject(this, &UStatusBarSubsystem::OnOutputLogOpened);
-	OutputLogDrawer.OnDrawerDismissedDelegate.BindUObject(this, &UStatusBarSubsystem::OnOutputLogDismised);
-	OutputLogDrawer.CustomWidget = OutputLog;
+		OutputLogDrawer.GetDrawerContentDelegate.BindUObject(this, &UStatusBarSubsystem::OnGetOutputLog);
+		OutputLogDrawer.OnDrawerOpenedDelegate.BindUObject(this, &UStatusBarSubsystem::OnOutputLogOpened);
+		OutputLogDrawer.OnDrawerDismissedDelegate.BindUObject(this, &UStatusBarSubsystem::OnOutputLogDismised);
+		OutputLogDrawer.CustomWidget = OutputLog;
 
-	OutputLogDrawer.ButtonText = LOCTEXT("StatusBar_OutputLogButton", "Output Log");
-	OutputLogDrawer.ToolTipText = FText::Format(LOCTEXT("StatusBar_OutputLogButtonTip", "Opens the output log drawer. ({0}) cycles between focusing the console command box, opening the output log drawer, and closing it.\nThe output log drawer may also be toggled directly with ({1})"), FGlobalEditorCommonCommands::Get().OpenConsoleCommandBox->GetInputText(), FGlobalEditorCommonCommands::Get().OpenOutputLogDrawer->GetInputText());
-	OutputLogDrawer.Icon = FAppStyle::Get().GetBrush("Log.TabIcon");
+		OutputLogDrawer.ButtonText = LOCTEXT("StatusBar_OutputLogButton", "Output Log");
+		OutputLogDrawer.ToolTipText = FText::Format(LOCTEXT("StatusBar_OutputLogButtonTip", "Opens the output log drawer. ({0}) cycles between focusing the console command box, opening the output log drawer, and closing it.\nThe output log drawer may also be toggled directly with ({1})"), FGlobalEditorCommonCommands::Get().OpenConsoleCommandBox->GetInputText(), FGlobalEditorCommonCommands::Get().OpenOutputLogDrawer->GetInputText());
+		OutputLogDrawer.Icon = FAppStyle::Get().GetBrush("Log.TabIcon");
 
-	for (const TUniquePtr<IGlobalStatusBarExtension>& Extension : GlobalStatusBarExtensions)
-	{
-		Extension->ExtendOutputLogDrawer(OutputLogDrawer);
+		for (const TUniquePtr<IGlobalStatusBarExtension>& Extension : GlobalStatusBarExtensions)
+		{
+			Extension->ExtendOutputLogDrawer(OutputLogDrawer);
+		}
+
+		StatusBar->RegisterDrawer(MoveTemp(OutputLogDrawer));
 	}
 
-	StatusBar->RegisterDrawer(MoveTemp(OutputLogDrawer));
 
 	// Clean up stale status bars
 	for (auto It = StatusBars.CreateIterator(); It; ++It)
@@ -563,7 +574,7 @@ TSharedRef<SWidget> UStatusBarSubsystem::MakeStatusBarWidget(FName StatusBarName
 
 	FStatusBarData StatusBarData;
 	StatusBarData.StatusBarWidget = StatusBar;
-	StatusBarData.ConsoleEditBox = ConsoleEditBox;
+	StatusBarData.ConsoleEditBox = ConsoleEditBox ? ConsoleEditBox : SNullWidget::NullWidget.ToSharedPtr();
 
 	StatusBars.Add(StatusBarName, StatusBarData);
 
diff --git a/Engine/Source/Editor/StatusBar/Public/StatusBarSubsystem.h b/Engine/Source/Editor/StatusBar/Public/StatusBarSubsystem.h
--- a/Engine/Source/Editor/StatusBar/Public/StatusBarSubsystem.h
+++ b/Engine/Source/Editor/StatusBar/Public/StatusBarSubsystem.h
@@ -47,6 +47,19 @@ private:
 	int32 Id;
 };
 
+namespace UE::Editor::Toolbars
+{
+	/** Options used to specify what parts (if any) of the status bar should be hidden. */
+	enum class ECreateStatusBarOptions : uint8
+	{
+		Default = 0,
+		HideContentBrowser = 1,
+		HideOutputLog = 1 << 1,
+		HideSourceControl = 1 << 2
+	};
+	ENUM_CLASS_FLAGS(ECreateStatusBarOptions);
+}
+
 struct FStatusBarData
 {
 	TWeakPtr<SStatusBar> StatusBarWidget;
@@ -123,7 +136,7 @@ public:
 	 * @param StatusBarName	The name of the status bar for updating it later. This name must be unique. You can use the instance number of an fname to generate a unique fname for comparison but a non-unique one for serialization of status bar data. I.E all asset editors of a specific type will be saved the same but at runtime will be indentified uniquely. This is usally what you want.
 	 * @param ParentTab	Parent tab of the status bar.
 	 */
-	TSharedRef<SWidget> MakeStatusBarWidget(FName UniqueStatusBarName, const TSharedRef<SDockTab>& InParentTab);
+	TSharedRef<SWidget> MakeStatusBarWidget(FName UniqueStatusBarName, const TSharedRef<SDockTab>& InParentTab, UE::Editor::Toolbars::ECreateStatusBarOptions CreateStatusBarOptions = UE::Editor::Toolbars::ECreateStatusBarOptions::Default);
 
 	/**
 	 * @return true if a status bar was found for the active window
