[Progress Bars] Remove toast notification functionality for progress bars since it is just duplicate information

diff --git a/Engine/Source/Editor/StatusBar/Private/SStatusBar.cpp b/Engine/Source/Editor/StatusBar/Private/SStatusBar.cpp
--- a/Engine/Source/Editor/StatusBar/Private/SStatusBar.cpp
+++ b/Engine/Source/Editor/StatusBar/Private/SStatusBar.cpp
@@ -25,19 +25,12 @@
 #include "SWidgetDrawer.h"
 #include "Misc/ConfigCacheIni.h"
 #include "Widgets/Notifications/SProgressBar.h"
-#include "Widgets/Notifications/SNotificationList.h"
-#include "Framework/Notifications/NotificationManager.h"
 #include "Widgets/Notifications/INotificationWidget.h"
 
 #define LOCTEXT_NAMESPACE "StatusBar"
 
 namespace StatusBarNotificationConstants
 {
-	// How long progress notification toasts should appear for
-	const float NotificationExpireTime = 5.0f;
-
-	const float NotificationFadeDuration = .15f;
-
 	// Delay before a progress notification becomes visible. This is to avoid the status bar to animate and flicker from short lived notifications. 
 	const double NotificationDelay = 3.0;
 }
@@ -664,17 +657,6 @@ void SStatusBar::StartProgressNotification(FProgressNotificationHandle InHandle,
 		if(TotalWorkToDo > 0)
 		{
 			ProgressNotifications.Emplace(DisplayText, FPlatformTime::Seconds(), InHandle, TotalWorkToDo);
-
-			// If a notification was already active, refresh its fadeout time 
-			if (TSharedPtr<SNotificationItem> ActiveProgressNotificationPin = ActiveProgressNotification.Pin())
-			{
-				ActiveProgressNotificationPin->SetExpireDuration(StatusBarNotificationConstants::NotificationExpireTime);
-				ActiveProgressNotificationPin->ExpireAndFadeout();
-			}
-			else
-			{
-				bAllowedToRefreshProgressNotification = true;
-			}
 			UpdateProgressStatus();
 		}
 	}
@@ -940,47 +922,14 @@ void SStatusBar::UpdateProgressStatus()
 
 		bShouldAnyProgressBeVisible &= (NumIncompleteTasks > 0);
 		
-		TSharedPtr<SStatusBarProgressWidget> ActiveNotificationProgressWidgetPin = ActiveNotificationProgressWidget.Pin();
-
 		if(bShouldAnyProgressBeVisible)
 		{
-			TSharedPtr<SNotificationItem> ActiveProgressNotificationPin = ActiveProgressNotification.Pin();
-		
-			// Show a new notification the first time a new progress task is started assuming we don't already have a notification open
-			if (!ActiveProgressNotificationPin.IsValid() && bAllowedToRefreshProgressNotification)
-			{
-				FNotificationInfo ProgressNotification(FText::GetEmpty());
-
-				ActiveNotificationProgressWidgetPin = SNew(SStatusBarProgressWidget, true);
-				ProgressNotification.ContentWidget = ActiveNotificationProgressWidgetPin.ToSharedRef();
-				ProgressNotification.FadeOutDuration = StatusBarNotificationConstants::NotificationFadeDuration;
-				ProgressNotification.ForWindow = FSlateApplication::Get().FindWidgetWindow(AsShared());
-				ProgressNotification.bAllowThrottleWhenFrameRateIsLow = false;
-				ActiveProgressNotificationPin = FSlateNotificationManager::Get().AddNotification(ProgressNotification);
-
-				ActiveProgressNotification = ActiveProgressNotificationPin;
-				ActiveNotificationProgressWidget = ActiveNotificationProgressWidgetPin;
-
-				ActiveProgressNotificationPin->SetExpireDuration(StatusBarNotificationConstants::NotificationExpireTime);
-				ActiveProgressNotificationPin->ExpireAndFadeout();
-
-				// Do not show the notification again unless a new task is started.
-				bAllowedToRefreshProgressNotification = false;
-			}
-
 			OpenProgressBar();
 
 			ProgressBar->SetPercent(Percent);
 
 			ProgressBar->SetProgressText(StatusBarProgressText);
 		}
-
-		// Update the notification if it exists. Even if no progress should be visible, if the notification is visible we want to update it while it fades out
-		if (ActiveNotificationProgressWidgetPin)
-		{
-			ActiveNotificationProgressWidgetPin->SetProgressText(StatusBarProgressText);
-			ActiveNotificationProgressWidgetPin->SetProgressPercent(Percent);
-		}
 	}
 
 	if (NumIncompleteTasks == 0)
@@ -997,7 +946,6 @@ void SStatusBar::OpenProgressBar()
 void SStatusBar::DismissProgressBar()
 {
 	ProgressBar->DismissProgressBar();
-	bAllowedToRefreshProgressNotification = false;
 	ProgressNotifications.Empty();
 }
 
diff --git a/Engine/Source/Editor/StatusBar/Private/SStatusBar.h b/Engine/Source/Editor/StatusBar/Private/SStatusBar.h
--- a/Engine/Source/Editor/StatusBar/Private/SStatusBar.h
+++ b/Engine/Source/Editor/StatusBar/Private/SStatusBar.h
@@ -194,14 +194,10 @@ private:
 	TWeakPtr<SDockTab> ParentTab;
 	TSharedPtr<SWidgetDrawer> WidgetDrawer;
 	TSharedPtr<SStatusBarProgressArea> ProgressBar;
-	TWeakPtr<SNotificationItem> ActiveProgressNotification;
-	TWeakPtr<SStatusBarProgressWidget> ActiveNotificationProgressWidget;
 
 	const FSlateBrush* UpArrow = nullptr;
 	const FSlateBrush* DownArrow = nullptr;
 
-	bool bAllowedToRefreshProgressNotification = false;
-
 	/** Whether or not to create the source control section. */
 	bool bCreateSourceControlSection = true;
 };
