SourceControl: Fixed an unnecessary 'Check Out' dialog popping up when doing a 'Revert All Files'.

Set the bFastSave to true, which reads:
"true if we should do a fast save. (I.E dont prompt the user to save, dont prompt for checkout, and only save packages that are currently writable)."

#rnx

diff --git a/Engine/Source/Editor/StatusBar/Private/SourceControlMenuHelpers.cpp b/Engine/Source/Editor/StatusBar/Private/SourceControlMenuHelpers.cpp
--- a/Engine/Source/Editor/StatusBar/Private/SourceControlMenuHelpers.cpp
+++ b/Engine/Source/Editor/StatusBar/Private/SourceControlMenuHelpers.cpp
@@ -193,7 +193,7 @@ void FSourceControlCommands::RevertAllModifiedFiles_Clicked()
 		const bool bPromptUserToSave = false;
 		const bool bSaveMapPackages = true;
 		const bool bSaveContentPackages = true;
-		const bool bFastSave = false;
+		const bool bFastSave = true;
 		const bool bNotifyNoPackagesSaved = false;
 		const bool bCanBeDeclined = false;
 		FEditorFileUtils::SaveDirtyPackages(bPromptUserToSave, bSaveMapPackages, bSaveContentPackages, bFastSave, bNotifyNoPackagesSaved, bCanBeDeclined);
