SourceControl: Fixed file not being saved to disk when using 'Revert All Files' - issue was introduced in #41524194.

The fast save doesn't save to disk if there is no file. So if adding a new actor (but not saving it to disk yet) and doing 'Revert All Files', it would not be reverted.

#rnx

diff --git a/Engine/Source/Editor/StatusBar/Private/SourceControlMenuHelpers.cpp b/Engine/Source/Editor/StatusBar/Private/SourceControlMenuHelpers.cpp
--- a/Engine/Source/Editor/StatusBar/Private/SourceControlMenuHelpers.cpp
+++ b/Engine/Source/Editor/StatusBar/Private/SourceControlMenuHelpers.cpp
@@ -230,7 +230,7 @@ void FSourceControlCommands::RevertAllModifiedFiles_Clicked()
 		const bool bPromptUserToSave = false;
 		const bool bSaveMapPackages = true;
 		const bool bSaveContentPackages = true;
-		const bool bFastSave = true;
+		const bool bFastSave = false;
 		const bool bNotifyNoPackagesSaved = false;
 		const bool bCanBeDeclined = false;
 		FEditorFileUtils::SaveDirtyPackages(bPromptUserToSave, bSaveMapPackages, bSaveContentPackages, bFastSave, bNotifyNoPackagesSaved, bCanBeDeclined);
