Rename Snapshot History for Branch History

#rnx
#Synced-CL 48697675

diff --git a/Engine/Source/Editor/MainFrame/Private/Frame/MainFrameActions.cpp b/Engine/Source/Editor/MainFrame/Private/Frame/MainFrameActions.cpp
--- a/Engine/Source/Editor/MainFrame/Private/Frame/MainFrameActions.cpp
+++ b/Engine/Source/Editor/MainFrame/Private/Frame/MainFrameActions.cpp
@@ -108,7 +108,7 @@ void FMainFrameCommands::RegisterCommands()
 	UI_COMMAND( ViewChangelists, "View Changes", "Opens a dialog displaying current changes.", EUserInterfaceActionType::Button, FInputChord() );
 	ActionList->MapAction(ViewChangelists, FExecuteAction::CreateStatic( &FMainFrameActionCallbacks::ViewChangelists ), FCanExecuteAction::CreateStatic( &FMainFrameActionCallbacks::CanViewChangelists ) );
 
-	UI_COMMAND( ViewSnapshotHistory, "Open Snapshot History", "See all the changes that have been made to this project over time.", EUserInterfaceActionType::Button, FInputChord());
+	UI_COMMAND( ViewSnapshotHistory, "Open Branch History", "See all the changes that have been made to this branch over time.", EUserInterfaceActionType::Button, FInputChord());
 	ActionList->MapAction( ViewSnapshotHistory, FExecuteAction::CreateStatic( &FMainFrameActionCallbacks::ViewSnapshotHistory ), FCanExecuteAction::CreateStatic( &FMainFrameActionCallbacks::CanViewSnapshotHistory ) );
 
 	UI_COMMAND( ViewBranchExplorer, "Open Branch Explorer", "See all the branches which this project uses.", EUserInterfaceActionType::Button, FInputChord() );
diff --git a/Engine/Source/Editor/StatusBar/Private/SourceControlMenuHelpers.cpp b/Engine/Source/Editor/StatusBar/Private/SourceControlMenuHelpers.cpp
--- a/Engine/Source/Editor/StatusBar/Private/SourceControlMenuHelpers.cpp
+++ b/Engine/Source/Editor/StatusBar/Private/SourceControlMenuHelpers.cpp
@@ -53,7 +53,7 @@ void FSourceControlCommands::RegisterCommands()
 	UI_COMMAND(ChangeSourceControlSettings, "Change Revision Control Settings...", "Opens a dialog to change revision control settings.", EUserInterfaceActionType::Button, FInputChord());
 	UI_COMMAND(ViewChangelists, "View Changes", "Opens a dialog displaying current changes.", EUserInterfaceActionType::Button, FInputChord());
 	UI_COMMAND(SubmitContent, "Submit Content", "Opens a dialog with check in options for content and levels.", EUserInterfaceActionType::Button, FInputChord());
-	UI_COMMAND(ViewSnapshotHistory, "Open Snapshot History", "See all the changes that have been made to this project over time.", EUserInterfaceActionType::Button, FInputChord());
+	UI_COMMAND(ViewSnapshotHistory, "Open Branch History", "See all the changes that have been made to this branch over time.", EUserInterfaceActionType::Button, FInputChord());
 	UI_COMMAND(ViewBranchExplorer, "Open Branch Explorer", "See all the branches which this project uses.", EUserInterfaceActionType::Button, FInputChord());
 	UI_COMMAND(CheckOutModifiedFiles, "Check Out Modified Files", "Opens a dialog to check out any assets which have been modified.", EUserInterfaceActionType::Button, FInputChord());
 	UI_COMMAND(RevertAll, "Revert All Files", "Opens a dialog to revert any assets which have been modified.", EUserInterfaceActionType::Button, FInputChord());
