Skein / URC2: Replaced usage of 'snapshot' with 'revision'.

#rnx
#Synced-CL 48897982

diff --git a/Engine/Source/Editor/StatusBar/Private/SourceControlMenuHelpers.cpp b/Engine/Source/Editor/StatusBar/Private/SourceControlMenuHelpers.cpp
--- a/Engine/Source/Editor/StatusBar/Private/SourceControlMenuHelpers.cpp
+++ b/Engine/Source/Editor/StatusBar/Private/SourceControlMenuHelpers.cpp
@@ -254,7 +254,7 @@ bool FSourceControlCommands::RevertAllModifiedFiles_CanExecute()
 
 void FSourceControlCommands::RevertAllModifiedFiles_Clicked()
 {
-	FText Message = LOCTEXT("RevertAllModifiedFiles", "Are you sure you want to revert all local changes? By proceeding, your local changes will be discarded, and the state of your last synced snapshot will be restored.");
+	FText Message = LOCTEXT("RevertAllModifiedFiles", "Are you sure you want to revert all local changes? By proceeding, your local changes will be discarded, and the state of your last synced revision will be restored.");
 	FText Title = LOCTEXT("RevertAllModifiedFiles_Title", "Revert all local changes");
 	if (FMessageDialog::Open(EAppMsgType::YesNo, EAppReturnType::No, Message, Title) == EAppReturnType::Yes)
 	{
