Branch Explorer: Allow a keyboard shortcut (Ctrl+Alt+B) to open the Branch Explorer drawer
We already have the ViewBranchExplorer UI command so we add the default keyboard shortcut and then add the command to the global level editor actions so they respond to the shortcut.

#Synced-CL 49027723

diff --git a/Engine/Source/Editor/StatusBar/Private/SourceControlMenuHelpers.cpp b/Engine/Source/Editor/StatusBar/Private/SourceControlMenuHelpers.cpp
--- a/Engine/Source/Editor/StatusBar/Private/SourceControlMenuHelpers.cpp
+++ b/Engine/Source/Editor/StatusBar/Private/SourceControlMenuHelpers.cpp
@@ -29,6 +29,7 @@
 #include "Styling/StyleColors.h"
 #include "HAL/IConsoleManager.h"
 #include "SSourceControlControls.h"
+#include "LevelEditor.h"
 
 #define LOCTEXT_NAMESPACE "SourceControlCommands"
 
@@ -54,7 +55,7 @@ void FSourceControlCommands::RegisterCommands()
 	UI_COMMAND(ViewChangelists, "View Changes", "Opens a dialog displaying current changes.", EUserInterfaceActionType::Button, FInputChord());
 	UI_COMMAND(SubmitContent, "Submit Content", "Opens a dialog with check in options for content and levels.", EUserInterfaceActionType::Button, FInputChord());
 	UI_COMMAND(ViewSnapshotHistory, "Open Branch History", "See all the changes that have been made to this branch over time.", EUserInterfaceActionType::Button, FInputChord());
-	UI_COMMAND(ViewBranchExplorer, "Open Branch Explorer", "See all the branches which this project uses.", EUserInterfaceActionType::Button, FInputChord());
+	UI_COMMAND(ViewBranchExplorer, "Open Branch Explorer", "See all the branches which this project uses.", EUserInterfaceActionType::Button, FInputChord(EModifierKey::Control | EModifierKey::Alt, EKeys::B));
 	UI_COMMAND(CheckOutModifiedFiles, "Check Out Modified Files", "Opens a dialog to check out any assets which have been modified.", EUserInterfaceActionType::Button, FInputChord());
 	UI_COMMAND(RevertAll, "Revert All Files", "Opens a dialog to revert any assets which have been modified.", EUserInterfaceActionType::Button, FInputChord());
 
@@ -101,6 +102,18 @@ void FSourceControlCommands::RegisterCommands()
 		FIsActionButtonVisible::CreateStatic(&FSourceControlCommands::ViewBranchExplorer_IsVisible)
 	);
 
+	// Add to global command list to allow key shortcut anywhere in Editor.
+	FLevelEditorModule& LevelEditorModule = FModuleManager::LoadModuleChecked<FLevelEditorModule>("LevelEditor");
+	TSharedPtr<FUICommandList> CommandList = LevelEditorModule.GetGlobalLevelEditorActions();
+	if (CommandList.IsValid())
+	{
+		CommandList->MapAction(
+			ViewBranchExplorer,
+			FExecuteAction::CreateStatic(&FSourceControlCommands::ViewBranchExplorer_Clicked),
+			FCanExecuteAction::CreateStatic(&FSourceControlCommands::ViewBranchExplorer_CanExecute)
+		);
+	}
+
 	ActionList->MapAction(
 		CheckOutModifiedFiles,
 		FExecuteAction::CreateStatic(&FSourceControlCommands::CheckOutModifiedFiles_Clicked),
