SourceControl: Mark 'Branch Explorer' feature as beta in UI.

#rnx

#Synced-CL 50095366

diff --git a/Engine/Source/Editor/MainFrame/Private/Frame/MainFrameActions.cpp b/Engine/Source/Editor/MainFrame/Private/Frame/MainFrameActions.cpp
--- a/Engine/Source/Editor/MainFrame/Private/Frame/MainFrameActions.cpp
+++ b/Engine/Source/Editor/MainFrame/Private/Frame/MainFrameActions.cpp
@@ -112,7 +112,7 @@ void FMainFrameCommands::RegisterCommands()
 	UI_COMMAND( ViewSnapshotHistory, "Open Branch History", "See all the changes that have been made to this branch over time.", EUserInterfaceActionType::Button, FInputChord());
 	ActionList->MapAction( ViewSnapshotHistory, FExecuteAction::CreateStatic( &FMainFrameActionCallbacks::ViewSnapshotHistory ), FCanExecuteAction::CreateStatic( &FMainFrameActionCallbacks::CanViewSnapshotHistory ) );
 
-	UI_COMMAND( ViewBranchExplorer, "Open Branch Explorer", "See all the branches which this project uses.", EUserInterfaceActionType::Button, FInputChord() );
+	UI_COMMAND( ViewBranchExplorer, "Open Branch Explorer (Beta)", "See all the branches which this project uses.", EUserInterfaceActionType::Button, FInputChord() );
 	ActionList->MapAction( ViewBranchExplorer, FExecuteAction::CreateStatic( &FMainFrameActionCallbacks::ViewBranchExplorer ), FCanExecuteAction::CreateStatic( &FMainFrameActionCallbacks::CanViewBranchExplorer ), FGetActionCheckState(), FIsActionButtonVisible::CreateStatic( &FMainFrameActionCallbacks::IsVisibleBranchExplorer ) );
 
 	UI_COMMAND( SubmitContent, "Submit Content", "Opens a dialog with check in options for content and levels.", EUserInterfaceActionType::Button, FInputChord() );
diff --git a/Engine/Source/Editor/StatusBar/Private/SourceControlMenuHelpers.cpp b/Engine/Source/Editor/StatusBar/Private/SourceControlMenuHelpers.cpp
--- a/Engine/Source/Editor/StatusBar/Private/SourceControlMenuHelpers.cpp
+++ b/Engine/Source/Editor/StatusBar/Private/SourceControlMenuHelpers.cpp
@@ -54,7 +54,7 @@ void FSourceControlCommands::RegisterCommands()
 	UI_COMMAND(ViewChangelists, "View Changes", "Opens a dialog displaying current changes.", EUserInterfaceActionType::Button, FInputChord());
 	UI_COMMAND(SubmitContent, "Submit Content", "Opens a dialog with check in options for content and levels.", EUserInterfaceActionType::Button, FInputChord());
 	UI_COMMAND(ViewSnapshotHistory, "Open Branch History", "See all the changes that have been made to this branch over time.", EUserInterfaceActionType::Button, FInputChord());
-	UI_COMMAND(ViewBranchExplorer, "Open Branch Explorer", "See all the branches which this project uses.", EUserInterfaceActionType::Button, FInputChord(EModifierKey::Control | EModifierKey::Alt, EKeys::B));
+	UI_COMMAND(ViewBranchExplorer, "Open Branch Explorer (Beta)", "See all the branches which this project uses.", EUserInterfaceActionType::Button, FInputChord(EModifierKey::Control | EModifierKey::Alt, EKeys::B));
 	UI_COMMAND(CheckOutModifiedFiles, "Check Out Modified Files", "Opens a dialog to check out any assets which have been modified.", EUserInterfaceActionType::Button, FInputChord());
 	UI_COMMAND(RevertAll, "Revert All Files", "Opens a dialog to revert any assets which have been modified.", EUserInterfaceActionType::Button, FInputChord());
 
